AudioLink Abstractions and Engine Modules

diff --git a/Engine/Source/Editor/AudioEditor/AudioEditor.Build.cs b/Engine/Source/Editor/AudioEditor/AudioEditor.Build.cs
--- a/Engine/Source/Editor/AudioEditor/AudioEditor.Build.cs
+++ b/Engine/Source/Editor/AudioEditor/AudioEditor.Build.cs
@@ -27,7 +27,8 @@ public class AudioEditor : ModuleRules
 				"UMG",
 				"DeveloperSettings",
 				"UMGEditor",
-				"AudioExtensions"
+				"AudioExtensions",
+				"AudioLinkEngine"
 			}
 		);
 
diff --git a/Engine/Source/Runtime/AudioLink/AudioLinkCore/AudioLinkCore.Build.cs b/Engine/Source/Runtime/AudioLink/AudioLinkCore/AudioLinkCore.Build.cs
new file mode 100644
--- /dev/null
+++ b/Engine/Source/Runtime/AudioLink/AudioLinkCore/AudioLinkCore.Build.cs
@@ -0,0 +1,18 @@
+// Copyright Epic Games, Inc. All Rights Reserved.
+
+namespace UnrealBuildTool.Rules
+{
+	public class AudioLinkCore : ModuleRules
+	{
+		public AudioLinkCore (ReadOnlyTargetRules Target) : base(Target)
+		{
+			PublicDependencyModuleNames.AddRange(
+				new string[]
+				{
+					"Core",
+					"CoreUObject",
+				}
+			);
+		}
+	}
+}
diff --git a/Engine/Source/Runtime/AudioLink/AudioLinkCore/Public/AudioLinkSettingsAbstract.h b/Engine/Source/Runtime/AudioLink/AudioLinkCore/Public/AudioLinkSettingsAbstract.h
new file mode 100644
--- /dev/null
+++ b/Engine/Source/Runtime/AudioLink/AudioLinkCore/Public/AudioLinkSettingsAbstract.h
@@ -0,0 +1,52 @@
+// Copyright Epic Games, Inc. All Rights Reserved.
+
+#pragma once
+
+#include "CoreMinimal.h"
+#include "Templates/SharedPointer.h"
+#include "AudioLinkSettingsAbstract.generated.h"
+
+/**
+  * This interface should be used to provide a non-uclass version of the data described in
+  * your implementation of UAudioLinkSettingsAbstract
+  */
+class AUDIOLINKCORE_API IAudioLinkSettingsProxy
+{
+public:
+	virtual ~IAudioLinkSettingsProxy() = default;
+};
+
+/**
+  * This opaque class should be used for specifying settings for how audio should be
+  * send to an external audio link.
+  */
+UCLASS(config = Engine, abstract, editinlinenew, BlueprintType)
+class AUDIOLINKCORE_API UAudioLinkSettingsAbstract : public UObject
+{
+	GENERATED_BODY()
+
+public:
+	using FSharedSettingsProxyPtr = TSharedPtr<IAudioLinkSettingsProxy, ESPMode::ThreadSafe>;
+
+	virtual FName GetFactoryName() const PURE_VIRTUAL(UAudioLinkSettingsAbstract::GetFactoryName, return {};);
+	
+	virtual const FSharedSettingsProxyPtr& GetProxy() const
+	{
+		if (!ProxyInstance.IsValid())
+		{
+			ProxyInstance = MakeProxy();
+		}
+		return ProxyInstance;
+	}
+
+	template<typename T> 
+	auto GetCastProxy() const
+	{
+		return StaticCastSharedPtr<T>(GetProxy());
+	}
+
+protected:
+	virtual FSharedSettingsProxyPtr MakeProxy() const PURE_VIRTUAL(UAudioLinkSettingsAbstract::MakeProxy, return {};);
+	mutable FSharedSettingsProxyPtr ProxyInstance;
+};
+
diff --git a/Engine/Source/Runtime/AudioLink/AudioLinkEngine/AudioLinkEngine.Build.cs b/Engine/Source/Runtime/AudioLink/AudioLinkEngine/AudioLinkEngine.Build.cs
new file mode 100644
--- /dev/null
+++ b/Engine/Source/Runtime/AudioLink/AudioLinkEngine/AudioLinkEngine.Build.cs
@@ -0,0 +1,20 @@
+// Copyright Epic Games, Inc. All Rights Reserved.
+
+namespace UnrealBuildTool.Rules
+{
+	public class AudioLinkEngine : ModuleRules
+	{
+		public AudioLinkEngine(ReadOnlyTargetRules Target) : base(Target)
+		{
+			PrivateDependencyModuleNames.AddRange(
+				new string[]
+				{
+					// We depend on USoundSubmix/AudioDevice which is in Engine.
+					"Core",
+					"CoreUObject",
+					"Engine",
+				}
+			);
+		}
+	}
+}
diff --git a/Engine/Source/Runtime/AudioLink/AudioLinkEngine/Private/AudioLinkModule.cpp b/Engine/Source/Runtime/AudioLink/AudioLinkEngine/Private/AudioLinkModule.cpp
new file mode 100644
--- /dev/null
+++ b/Engine/Source/Runtime/AudioLink/AudioLinkEngine/Private/AudioLinkModule.cpp
@@ -0,0 +1,21 @@
+// Copyright Epic Games, Inc. All Rights Reserved.
+
+#include "CoreMinimal.h"
+#include "AudioLinkLog.h"
+#include "Modules/ModuleManager.h"
+
+DEFINE_LOG_CATEGORY(LogAudioLink);
+
+class FAudioLinkModule : public IModuleInterface
+{
+public:
+	// Begin IModuleInterface
+	virtual void StartupModule() override
+	{
+	}
+	virtual void ShutdownModule() override
+	{
+	}
+};
+
+IMPLEMENT_MODULE(FAudioLinkModule, AudioLink);
diff --git a/Engine/Source/Runtime/AudioLink/AudioLinkEngine/Private/BufferedListenerBase.cpp b/Engine/Source/Runtime/AudioLink/AudioLinkEngine/Private/BufferedListenerBase.cpp
new file mode 100644
--- /dev/null
+++ b/Engine/Source/Runtime/AudioLink/AudioLinkEngine/Private/BufferedListenerBase.cpp
@@ -0,0 +1,90 @@
+// Copyright Epic Games, Inc. All Rights Reserved.
+
+#include "BufferedListenerBase.h"
+#include "AudioLinkLog.h"
+
+FBufferedListenerBase::FBufferedListenerBase(int32 InDefaultCircularBufferSize)
+{
+	LocklessCircularBuffer.SetCapacity(InDefaultCircularBufferSize);
+}
+
+/** CONSUMER Thread. This will be called by the consuming data of the buffers. */
+bool FBufferedListenerBase::PopBuffer(float* InBuffer, int32 InBufferSizeInSamples, int32& OutSamplesWritten)
+{
+	OutSamplesWritten = LocklessCircularBuffer.Pop(InBuffer, InBufferSizeInSamples);
+	return bStarted;
+}
+
+/** CONSUMER Thread. This will be called by the consuming data of the buffers. */
+bool FBufferedListenerBase::GetFormat(IBufferedAudioOutput::FBufferFormat& OutFormat) const
+{
+	FReadScopeLock ReadLock(FormatKnownRwLock);
+	if (KnownFormat)
+	{
+		OutFormat = *KnownFormat;
+		return true; // success.
+	}
+	return false;
+}
+
+void FBufferedListenerBase::SetFormatKnownDelegate(FOnFormatKnown InFormatKnownDelegate)
+{
+	OnFormatKnown = InFormatKnownDelegate;
+}
+
+/** AUDIO MIXER THREAD. */
+void FBufferedListenerBase::OnBufferRecieved(const FBufferFormat& InFormat, TArrayView<const float> InBuffer)
+{
+	// Keep track of if we need to fire the delegate, so we can do it outside of a lock.
+	bool bFireFormatKnownDelegate = false;
+
+	// Do we know the format yet? (do this under a read-only lock, unless we need to change it).
+	{
+		// Read lock to check if we know the state.
+		FRWScopeLock Lock(FormatKnownRwLock, SLT_ReadOnly);
+
+		// Format known?
+		if (!KnownFormat)
+		{
+			{
+				// Upgrade lock to Write lock. (Safe here as we are the only writer, and only readers can race between.).
+				Lock.ReleaseReadOnlyLockAndAcquireWriteLock_USE_WITH_CAUTION();
+
+				// Write this format as our known format.
+				KnownFormat = InFormat;
+
+				// Remember to fire the delegate.
+				bFireFormatKnownDelegate = true;
+			}
+		}
+		else
+		{
+			// Sanity check they haven't changed on the source since it started.
+			ensure(InFormat == *KnownFormat);
+		}
+	}
+
+	// Fire format known delegate. (important this is done outside of the read/write lock above as it calls GetFormat, which needs a read lock).
+	if (bFireFormatKnownDelegate)
+	{
+		// Broadcast to consumer that we know our format and block rate.
+		OnFormatKnown.ExecuteIfBound(InFormat);
+	}
+
+	// Push the data into the circular buffer.
+	int32 SamplesPushed = LocklessCircularBuffer.Push(InBuffer.GetData(), InBuffer.Num());
+
+	// Warn of not enough space in circular buffer
+	if (SamplesPushed != InBuffer.Num())
+	{
+		// Prevent log spam by limiting to 1:100 logs.
+		static const int32 NumLogMessagesToSkip = 100;
+		static int32 LogPacifier = 0;
+		UE_CLOG(LogPacifier++ % NumLogMessagesToSkip == 0, LogAudioLink, Warning, TEXT("Overflow by '%d' Samples in Buffer Listener"), InBuffer.Num() - SamplesPushed);
+	}
+}
+
+void FBufferedListenerBase::Reserve(int32 InNumSamplesToReserve)
+{
+	LocklessCircularBuffer.SetCapacity(InNumSamplesToReserve);
+}
\ No newline at end of file
diff --git a/Engine/Source/Runtime/AudioLink/AudioLinkEngine/Private/BufferedListenerBase.h b/Engine/Source/Runtime/AudioLink/AudioLinkEngine/Private/BufferedListenerBase.h
new file mode 100644
--- /dev/null
+++ b/Engine/Source/Runtime/AudioLink/AudioLinkEngine/Private/BufferedListenerBase.h
@@ -0,0 +1,44 @@
+// Copyright Epic Games, Inc. All Rights Reserved.
+
+#pragma once
+
+#include "AudioDevice.h"
+#include "SignalProcessing/Public/DSP/Dsp.h"
+
+#include "IBufferedAudioOutput.h"
+#include "Misc/Optional.h"
+#include "Containers/ArrayView.h"
+
+/** Common base class of Buffered Listener objects.
+*/
+class AUDIOLINKENGINE_API FBufferedListenerBase : public IBufferedAudioOutput
+{
+protected:
+	FBufferedListenerBase(int32 InDefaultCircularBufferSize);
+	virtual ~FBufferedListenerBase() = default;
+
+	//~ Begin IBufferedAudioOutput
+	bool PopBuffer(float* InBuffer, int32 InBufferSizeInSamples, int32& OutSamplesWritten) override;
+	bool GetFormat(IBufferedAudioOutput::FBufferFormat& OutFormat) const override;
+	void Reserve(int32 InNumSamplesToReserve) override;
+	void SetFormatKnownDelegate(FOnFormatKnown InFormatKnownDelegate) override;
+	//~ End IBufferedAudioOutput
+
+	//* Common path to receive a new buffer, call from derived classes */
+	void OnBufferRecieved(const FBufferFormat& InFormat, TArrayView<const float> InBuffer);
+
+	/** Lock-less buffer to hold the data for the single source we're listening to, interleaved. */
+	Audio::TCircularAudioBuffer<float> LocklessCircularBuffer;
+
+	/** Read/Write slim lock protects format known optional below. */
+	mutable FRWLock FormatKnownRwLock;
+
+	/** Optional that holds the buffer format, if (and when) it's known. Protected by r/w Slim-lock */
+	TOptional<FBufferFormat> KnownFormat;
+
+	/** Delegate that fires when the format it known. Normally on the first buffer received. */
+	FOnFormatKnown OnFormatKnown;
+
+	/** Atomic flag we're been started yet. */
+	std::atomic<bool> bStarted;
+};
diff --git a/Engine/Source/Runtime/AudioLink/AudioLinkEngine/Private/BufferedSourceListener.cpp b/Engine/Source/Runtime/AudioLink/AudioLinkEngine/Private/BufferedSourceListener.cpp
new file mode 100644
--- /dev/null
+++ b/Engine/Source/Runtime/AudioLink/AudioLinkEngine/Private/BufferedSourceListener.cpp
@@ -0,0 +1,78 @@
+// Copyright Epic Games, Inc. All Rights Reserved.
+
+#include "BufferedSourceListener.h"
+
+/** Buffered Source Listener. */
+FBufferedSourceListener::FBufferedSourceListener(int32 InDefaultCircularBufferSize)
+	: FBufferedListenerBase{ InDefaultCircularBufferSize }
+{
+}
+
+FBufferedSourceListener::~FBufferedSourceListener()
+{
+	// No need to stop as its only record keeping and nothing to unregister.
+}
+
+bool FBufferedSourceListener::Start(FAudioDevice* InAudioDevice)
+{
+	if (ensure(!bStarted))
+	{
+		bStarted = true;
+		return true;
+	}
+	return false;
+}
+
+void FBufferedSourceListener::Stop(FAudioDevice* InAudioDevice)
+{
+	if (bStarted)
+	{
+		bStarted = false;
+	}
+}
+
+/** AUDIO MIXER THREAD. When a source is finished and returned to the pool, this call will be called. */
+void FBufferedSourceListener::OnSourceReleased(const int32 InSourceId)
+{	
+	if (InSourceId == CurrentSourceId)
+	{
+		CurrentSourceId = INDEX_NONE;
+		
+		// Reset the format, now our source we were listening to has stopped.
+		FWriteScopeLock WriteLock(FormatKnownRwLock);
+		KnownFormat.Reset();
+	}
+}
+
+/** AUDIO MIXER THREAD. New Audio buffers from the active sources enter here. */
+void FBufferedSourceListener::OnNewBuffer(const ISourceBufferListener::FOnNewBufferParams& InParams)
+{	
+	if (bStarted) 
+	{
+		// Make sure caller is sane.
+		check(InParams.AudioData);
+		check(InParams.NumChannels > 0);
+		check(InParams.NumSamples > 0);
+		check(InParams.SampleRate > 0);
+		check(InParams.SourceId >= 0);
+		
+		// Multiple sources can call this same instance, so make sure we're only listening to one.
+		int32 OriginalCurrentSourceId = CurrentSourceId.load(std::memory_order_relaxed);	
+		if (OriginalCurrentSourceId == INDEX_NONE)
+		{
+			// Exchange CurrentSourceId, which is INDEX_NONE with InSourceId  (unless another thread beats us to it). 
+			int32 NotSetSourceId = INDEX_NONE;
+			if (!CurrentSourceId.compare_exchange_strong(/* expected */ NotSetSourceId, /* desired */ InParams.SourceId, std::memory_order_relaxed))
+			{
+				return;
+			}
+		}
+
+		// Call to base class to handle.
+		FBufferFormat NewFormat;
+		NewFormat.NumChannels = InParams.NumChannels;
+		NewFormat.NumSamplesPerBlock = InParams.NumSamples;
+		NewFormat.NumSamplesPerSec = InParams.SampleRate;
+		OnBufferRecieved(NewFormat, MakeArrayView(InParams.AudioData, InParams.NumSamples));
+	}
+}
diff --git a/Engine/Source/Runtime/AudioLink/AudioLinkEngine/Private/BufferedSourceListener.h b/Engine/Source/Runtime/AudioLink/AudioLinkEngine/Private/BufferedSourceListener.h
new file mode 100644
--- /dev/null
+++ b/Engine/Source/Runtime/AudioLink/AudioLinkEngine/Private/BufferedSourceListener.h
@@ -0,0 +1,36 @@
+// Copyright Epic Games, Inc. All Rights Reserved.
+
+#pragma once
+
+#include "AudioDevice.h"
+#include "BufferedListenerBase.h"
+
+/**
+	FBufferedSourceListener: Concrete implementation of both a buffer listener and buffered output.
+	Contains a circular buffer.
+*/
+class AUDIOLINKENGINE_API FBufferedSourceListener : public ISourceBufferListener, public FBufferedListenerBase
+{
+public:
+	FBufferedSourceListener(int32 InDefaultCircularBufferSize);
+	virtual ~FBufferedSourceListener();
+
+	/**
+	 * @param  InDevice Audio Device to register this listener with.
+	 * @param  InOnFirstBufferCallback Delegate to fire when the format of the buffer is known.
+	 *
+	 * @return  success true, false otherwise.
+	 */
+	bool Start(FAudioDevice* InDevice) override;
+
+	/** Stop the buffer by unregistering the buffer listener */
+	void Stop(FAudioDevice* InDevice) override;
+	
+private:
+	//~ Begin ISourceBufferListener
+	void OnNewBuffer(const ISourceBufferListener::FOnNewBufferParams& InParams) override;
+	void OnSourceReleased(const int32 SourceId) override;
+	//~ End ISourceBufferListener
+
+	std::atomic<int32> CurrentSourceId	= INDEX_NONE;		// r/w AudioMixer thread.	
+};
diff --git a/Engine/Source/Runtime/AudioLink/AudioLinkEngine/Private/BufferedSubmixListener.cpp b/Engine/Source/Runtime/AudioLink/AudioLinkEngine/Private/BufferedSubmixListener.cpp
new file mode 100644
--- /dev/null
+++ b/Engine/Source/Runtime/AudioLink/AudioLinkEngine/Private/BufferedSubmixListener.cpp
@@ -0,0 +1,75 @@
+// Copyright Epic Games, Inc. All Rights Reserved.
+
+#include "BufferedSubmixListener.h"
+
+/** Buffered Submix Listener. */
+FBufferedSubmixListener::FBufferedSubmixListener(int32 InDefaultCircularBufferSize, bool bInZeroInputBuffer)
+	: FBufferedListenerBase{ InDefaultCircularBufferSize }
+	, bZeroInputBuffer{ bInZeroInputBuffer }
+{
+}
+
+FBufferedSubmixListener::~FBufferedSubmixListener()
+{
+	//Stop(); // FIXME! Virtual call in destructor.	
+	//StopInternal();
+	// Perhaps we should assert that we're still registered.
+}
+
+bool FBufferedSubmixListener::Start(FAudioDevice* InAudioDevice)
+{
+	if (ensure(!bStarted))
+	{		
+		if (ensure(InAudioDevice))
+		{
+			// Inputs look valid, we can start.
+			DeviceId = InAudioDevice->DeviceID;
+			InAudioDevice->RegisterSubmixBufferListener(this);
+			bStarted = true;
+			return true;
+		}
+	}
+	return false;
+}
+
+void FBufferedSubmixListener::Stop(FAudioDevice* InAudioDevice)
+{
+	StopInternal(InAudioDevice);
+}
+
+void FBufferedSubmixListener::StopInternal(FAudioDevice* InAudioDevice)
+{
+	if (bStarted)
+	{
+		if (ensure(InAudioDevice))
+		{
+			if (ensure(InAudioDevice->DeviceID == DeviceId))
+			{
+				InAudioDevice->UnregisterSubmixBufferListener(this);
+				DeviceId = bStarted = false;
+			}
+		}
+	}
+}
+
+void FBufferedSubmixListener::OnNewSubmixBuffer(const USoundSubmix* OwningSubmix, float* AudioData, int32 InNumSamples, int32 InNumChannels, const int32 InSampleRate, double)
+{
+	if(bStarted)
+	{
+		// Call to base class to handle.
+		FBufferFormat NewFormat;
+		NewFormat.NumChannels = InNumChannels;
+		NewFormat.NumSamplesPerBlock = InNumSamples;
+		NewFormat.NumSamplesPerSec = InSampleRate;
+		OnBufferRecieved(NewFormat, MakeArrayView(AudioData, InNumSamples));
+
+		// Optionally, zero the buffer if we're asked to. This in the case where we're running both Unreal+Consumer renderers at once.
+		// NOTE: this is dangerous as there's a chance we're not the only listener registered on this Submix. And will cause
+		// listeners after us to have a silent buffer. Use with caution. 
+
+		if (bZeroInputBuffer)
+		{
+			FMemory::Memzero(AudioData, sizeof(float) * InNumSamples);
+		}	
+	}
+}
diff --git a/Engine/Source/Runtime/AudioLink/AudioLinkEngine/Private/BufferedSubmixListener.h b/Engine/Source/Runtime/AudioLink/AudioLinkEngine/Private/BufferedSubmixListener.h
new file mode 100644
--- /dev/null
+++ b/Engine/Source/Runtime/AudioLink/AudioLinkEngine/Private/BufferedSubmixListener.h
@@ -0,0 +1,40 @@
+// Copyright Epic Games, Inc. All Rights Reserved.
+
+#pragma once
+
+#include "BufferedListenerBase.h"
+
+/** Concrete Submix Buffer Listener  
+*/
+class AUDIOLINKENGINE_API FBufferedSubmixListener : public ISubmixBufferListener, 
+													public FBufferedListenerBase
+{
+public:
+	/** Constructor
+	 * @param  InDefaultCircularBufferSize: Size of the circular buffer in samples by default.
+	 * @param  bInZeroInputBuffer: The passed in buffer will be set zerod after we've buffered it. Not 
+	 */
+	FBufferedSubmixListener(int32 InDefaultCircularBufferSize, bool bInZeroInputBuffer);
+
+	virtual ~FBufferedSubmixListener();
+
+	/**
+	 * Starts the Submix buffer listener by registering it with the passed in Audio Device.
+	 * @param  InDevice Audio Device to register this submix listener with.
+	 * @param  InOnFirstBufferCallback Delegate to fire when the format of the buffer is known.
+	 *
+	 * @return  success true, false otherwise.
+	 */
+	bool Start(FAudioDevice* InDevice) override;
+	void Stop(FAudioDevice* InDevice) override;
+
+private:
+	void StopInternal(FAudioDevice* InDevice);
+
+	//~ Begin ISubmixBufferListener
+	void OnNewSubmixBuffer(const USoundSubmix* OwningSubmix, float* AudioData, int32 InNumSamples, int32 InNumChannels, const int32 InSampleRate, double) override;
+	//~ End ISubmixBufferListener
+
+	Audio::FDeviceId DeviceId;
+	bool bZeroInputBuffer = false;
+};
\ No newline at end of file
diff --git a/Engine/Source/Runtime/AudioLink/AudioLinkEngine/Private/IAudioLinkBlueprintInterface.cpp b/Engine/Source/Runtime/AudioLink/AudioLinkEngine/Private/IAudioLinkBlueprintInterface.cpp
new file mode 100644
--- /dev/null
+++ b/Engine/Source/Runtime/AudioLink/AudioLinkEngine/Private/IAudioLinkBlueprintInterface.cpp
@@ -0,0 +1,8 @@
+// Copyright Epic Games, Inc. All Rights Reserved.
+
+#include "IAudioLinkBlueprintInterface.h"
+
+UAudioLinkBlueprintInterface::UAudioLinkBlueprintInterface(FObjectInitializer const& InObjectInitializer)
+	: UInterface(InObjectInitializer)
+{
+}
diff --git a/Engine/Source/Runtime/AudioLink/AudioLinkEngine/Private/IAudioLinkFactory.cpp b/Engine/Source/Runtime/AudioLink/AudioLinkEngine/Private/IAudioLinkFactory.cpp
new file mode 100644
--- /dev/null
+++ b/Engine/Source/Runtime/AudioLink/AudioLinkEngine/Private/IAudioLinkFactory.cpp
@@ -0,0 +1,59 @@
+// Copyright Epic Games, Inc. All Rights Reserved.
+
+#include "IAudioLinkFactory.h"
+#include "Algo/Transform.h"
+
+// Concrete Buffer Listeners.
+#include "BufferedSubmixListener.h" 
+#include "BufferedSourceListener.h" 
+
+IAudioLinkFactory::IAudioLinkFactory()
+{
+	IModularFeatures::Get().RegisterModularFeature(GetModularFeatureName(), this);
+}
+
+IAudioLinkFactory::~IAudioLinkFactory()
+{
+	IModularFeatures::Get().UnregisterModularFeature(GetModularFeatureName(), this);
+}
+
+FSharedBufferedOutputPtr IAudioLinkFactory::CreateSourceBufferListener(const FSourceBufferListenerCreateParams& InSourceCreateParams)
+{
+	check(IsInGameThread()); 
+	auto SourceBufferListenerSP = MakeShared<FBufferedSourceListener, ESPMode::ThreadSafe>(InSourceCreateParams.SizeOfBufferInFrames);
+	InSourceCreateParams.AudioComponent->SetSourceBufferListener(SourceBufferListenerSP, InSourceCreateParams.bShouldZeroBuffer);
+	return SourceBufferListenerSP;
+}
+
+FSharedBufferedOutputPtr IAudioLinkFactory::CreateSubmixBufferListener(const FSubmixBufferListenerCreateParams& InSubmixCreateParams)
+{
+	return MakeShared<FBufferedSubmixListener, ESPMode::ThreadSafe>(InSubmixCreateParams.SizeOfBufferInFrames, InSubmixCreateParams.bShouldZeroBuffer);	
+}
+
+FName IAudioLinkFactory::GetModularFeatureName()
+{
+	static FName FeatureName = FName(TEXT("AudioLink Factory"));
+	return FeatureName;
+}
+
+TArray<IAudioLinkFactory*> IAudioLinkFactory::GetAllRegisteredFactories()
+{
+	return IModularFeatures::Get().GetModularFeatureImplementations<IAudioLinkFactory>(GetModularFeatureName());
+}
+
+TArray<FName> IAudioLinkFactory::GetAllRegisteredFactoryNames()
+{
+	TArray<FName> Names;
+	Algo::Transform(GetAllRegisteredFactories(), Names, [](IAudioLinkFactory* Factory) { return Factory->GetFactoryName(); });
+	return Names;
+}
+
+IAudioLinkFactory* IAudioLinkFactory::FindFactory(const FName InFactoryImplName)
+{
+	TArray<IAudioLinkFactory*> Factories = GetAllRegisteredFactories();
+	if (IAudioLinkFactory** Found = Factories.FindByPredicate([InFactoryImplName](IAudioLinkFactory* Factory) { return Factory->GetFactoryName() == InFactoryImplName; }))
+	{
+		return *Found;
+	}
+	return nullptr;
+}
diff --git a/Engine/Source/Runtime/AudioLink/AudioLinkEngine/Public/AudioLinkLog.h b/Engine/Source/Runtime/AudioLink/AudioLinkEngine/Public/AudioLinkLog.h
new file mode 100644
--- /dev/null
+++ b/Engine/Source/Runtime/AudioLink/AudioLinkEngine/Public/AudioLinkLog.h
@@ -0,0 +1,7 @@
+// Copyright Epic Games, Inc. All Rights Reserved.
+
+#pragma once
+
+#include "CoreMinimal.h"
+
+AUDIOLINKENGINE_API DECLARE_LOG_CATEGORY_EXTERN(LogAudioLink, Log, All);
diff --git a/Engine/Source/Runtime/AudioLink/AudioLinkEngine/Public/IAudioLink.h b/Engine/Source/Runtime/AudioLink/AudioLinkEngine/Public/IAudioLink.h
new file mode 100644
--- /dev/null
+++ b/Engine/Source/Runtime/AudioLink/AudioLinkEngine/Public/IAudioLink.h
@@ -0,0 +1,18 @@
+// Copyright Epic Games, Inc. All Rights Reserved.
+
+#pragma once
+
+#include "CoreMinimal.h"
+
+/** IAudioLink
+  *	Abstract interface for AudioLink instances. 
+  *	Purely opaque.
+  */
+class IAudioLink
+{
+protected:
+	IAudioLink() = default;
+public:
+	virtual ~IAudioLink() = default;
+};
+
diff --git a/Engine/Source/Runtime/AudioLink/AudioLinkEngine/Public/IAudioLinkBlueprintInterface.h b/Engine/Source/Runtime/AudioLink/AudioLinkEngine/Public/IAudioLinkBlueprintInterface.h
new file mode 100644
--- /dev/null
+++ b/Engine/Source/Runtime/AudioLink/AudioLinkEngine/Public/IAudioLinkBlueprintInterface.h
@@ -0,0 +1,32 @@
+// Copyright Epic Games, Inc. All Rights Reserved.
+
+#pragma once
+
+#include "UObject/Interface.h"
+#include "IAudioLinkBlueprintInterface.generated.h"
+
+UINTERFACE(BlueprintType, meta = (CannotImplementInterfaceInBlueprint))
+class AUDIOLINKENGINE_API UAudioLinkBlueprintInterface : public UInterface
+{
+	GENERATED_UINTERFACE_BODY()
+};
+
+class AUDIOLINKENGINE_API IAudioLinkBlueprintInterface : public IInterface
+{
+	GENERATED_IINTERFACE_BODY()
+
+public:	
+	UFUNCTION(BlueprintCallable, Category = "AudioLink")
+	virtual void SetLinkSound(USoundBase* NewSound) = 0;
+
+	UFUNCTION(BlueprintCallable, Category = "AudioLink")
+	virtual void PlayLink(float StartTime = 0.0f) = 0;
+
+	/** Stop an audio component's sound, issue any delegates if needed */
+	UFUNCTION(BlueprintCallable, Category = "AudioLink")
+	virtual void StopLink() = 0;
+		
+	UFUNCTION(BlueprintCallable, Category = "AudioLink")
+	virtual bool IsLinkPlaying() const = 0;
+};
+
diff --git a/Engine/Source/Runtime/AudioLink/AudioLinkEngine/Public/IAudioLinkFactory.h b/Engine/Source/Runtime/AudioLink/AudioLinkEngine/Public/IAudioLinkFactory.h
new file mode 100644
--- /dev/null
+++ b/Engine/Source/Runtime/AudioLink/AudioLinkEngine/Public/IAudioLinkFactory.h
@@ -0,0 +1,118 @@
+// Copyright Epic Games, Inc. All Rights Reserved.
+
+#pragma once
+
+#include "IAudioLink.h"
+#include "Features/IModularFeatures.h"
+#include "AudioDeviceManager.h"
+#include "AudioLinkSettingsAbstract.h"
+#include "Templates/SubclassOf.h"
+#include "IBufferedAudioOutput.h"
+
+// Forward declarations.
+class USoundSubmix;
+class UAudioLinkComponent;
+
+/**
+	IAudioLinkFactory
+	Abstract interface for AudioLink factory objects. 
+*/
+class AUDIOLINKENGINE_API IAudioLinkFactory : public IModularFeature
+{
+protected:
+	/** Constructor will auto register this instance. */
+	IAudioLinkFactory();
+public:
+	/** Destructor will auto unregister this instance. */
+	virtual ~IAudioLinkFactory();
+
+	/** 
+	 * Gets the name of this factory type. This will uniquely identity itself, so it can be found by the FindFactory call below
+	 * @return Name of factory
+	 */
+	virtual FName GetFactoryName() const = 0;
+
+	/**
+	 * Gets the type of of the settings object that this factory uses to define its specific settings. All settings for each
+	 * AudioLink factory will derive from UAudioLinkSettingsAbstract
+	 * @return Class of the Settings Object
+	 */
+	virtual TSubclassOf<UAudioLinkSettingsAbstract> GetSettingsClass() const = 0;
+	
+	struct FSourceBufferListenerCreateParams
+	{
+		int32 SizeOfBufferInFrames = 0;
+		bool bShouldZeroBuffer = false;
+		TWeakObjectPtr<UAudioComponent> AudioComponent;
+		TWeakObjectPtr<USceneComponent> OwningComponent;
+	};
+	virtual FSharedBufferedOutputPtr CreateSourceBufferListener(const FSourceBufferListenerCreateParams&);
+
+	struct FSubmixBufferListenerCreateParams
+	{
+		int32 SizeOfBufferInFrames = 0;
+		bool bShouldZeroBuffer = false;
+	};
+	virtual FSharedBufferedOutputPtr CreateSubmixBufferListener(const FSubmixBufferListenerCreateParams&);
+			
+	/**
+	 * Parameters use when creating a Submix Audio Link 
+	 */
+	struct FAudioLinkSubmixCreateArgs
+	{
+		TWeakObjectPtr<const USoundSubmix> Submix;
+		FAudioDevice* Device = nullptr;
+		TWeakObjectPtr<UAudioLinkSettingsAbstract> Settings;
+	};
+
+	/**
+	 * Create a Submix Audio Link.
+	 * @param InCreateArgs Arguments used to create the AudioLink instance
+	 * @return The newly created Link instance (if successful).
+	 */
+	virtual TUniquePtr<IAudioLink> CreateSubmixAudioLink(const FAudioLinkSubmixCreateArgs& InCreateArgs) = 0;
+
+	/**
+	 * Parameters use when creating a Source Audio Link 
+	 */
+	struct FAudioLinkSourceCreateArgs
+	{
+		TWeakObjectPtr<UAudioComponent> AudioComponent;
+		TWeakObjectPtr<USceneComponent> OwningComponent;
+		TWeakObjectPtr<UAudioLinkSettingsAbstract> Settings;
+	};
+
+	/**
+	 * Create a Source Audio Link.
+	 * @param InCreateArgs Arguments used to create the AudioLink instance
+	 * @return The newly created Link instance (if successful).
+	 */
+	virtual TUniquePtr<IAudioLink> CreateSourceAudioLink(const FAudioLinkSourceCreateArgs&) = 0;
+		
+	/**	
+	 * Get the name of all AudioLink factories in the Modular Features registry.
+	 * @return "AudioLink Factory"
+	 */
+	static FName GetModularFeatureName();
+
+	/**
+	 * Gets all registered factory instances.
+	 * @return Array of all factories.
+	 */
+	static TArray<IAudioLinkFactory*> GetAllRegisteredFactories();
+
+	/**
+	 * Gets all registered factory names
+	 * @return Array of all factory names
+	 */
+	static TArray<FName> GetAllRegisteredFactoryNames();
+	
+	/**
+	 * Gets all registered factory names
+	 * @return Array of all factory names
+	 */
+	static IAudioLinkFactory* FindFactory(const FName InFactoryName);
+
+protected:	
+};
+
diff --git a/Engine/Source/Runtime/AudioLink/AudioLinkEngine/Public/IBufferedAudioOutput.h b/Engine/Source/Runtime/AudioLink/AudioLinkEngine/Public/IBufferedAudioOutput.h
new file mode 100644
--- /dev/null
+++ b/Engine/Source/Runtime/AudioLink/AudioLinkEngine/Public/IBufferedAudioOutput.h
@@ -0,0 +1,69 @@
+// Copyright Epic Games, Inc. All Rights Reserved.
+
+#pragma once
+
+/**
+	Abstract interface for communication with outputting audio objects
+	Examples concrete implementation of these are (source, submix etc).
+*/
+
+class IBufferedAudioOutput 
+{
+protected:
+	IBufferedAudioOutput() = default;
+
+public:
+	virtual ~IBufferedAudioOutput() = default;
+
+	/**
+	 * The format of a buffer.
+	 */
+	struct FBufferFormat
+	{
+		int32 NumSamplesPerBlock = 0;
+		int32 NumChannels = 0;
+		int32 NumSamplesPerSec = 0;
+
+		bool operator==(const FBufferFormat& InRhs) const
+		{
+			return InRhs.NumChannels == NumChannels && 
+				InRhs.NumSamplesPerBlock == NumSamplesPerBlock && 
+				InRhs.NumSamplesPerSec == NumSamplesPerSec;
+		}
+	};
+
+	// Delegates.
+	DECLARE_DELEGATE_OneParam(FOnFormatKnown, FBufferFormat);
+	virtual void SetFormatKnownDelegate(FOnFormatKnown InFormatKnownDelegate) = 0;
+
+	virtual bool Start(FAudioDevice* InAudioDevice) = 0;
+	virtual void Stop(FAudioDevice* InAudioDevice) = 0;
+
+	/**
+	 * Attempts to Atomically copy a buffer sized amount of Buffered Sample data from the interface.
+	 * 
+	 * @param InBuffer Sample Buffer to Write to
+	 * @param InBufferSizeInSamples Buffer Size in total samples. 
+	 * @param OutSamplesWritten Number of Samples Written to the buffer
+	 * 
+	 * @return false no more data, true more data to come
+	 */
+	virtual bool PopBuffer(float* InBuffer, int32 InBufferSizeInSamples, int32& OutSamplesWritten) = 0;
+
+	/**
+	 * Gets the format of the buffer, if its known.
+	 *
+	 * @return  success true, false otherwise. 
+	 */
+	virtual bool GetFormat(FBufferFormat& OutFormat) const = 0;
+
+	/**
+	 * Reserve at least this many samples in buffer.
+	 *
+	 * @param InNumSamplesToReserve Reserve this number of samples.
+	 */
+	virtual void Reserve(int32 InNumSamplesToReserve) = 0;
+};
+
+using FWeakBufferedOutputPtr = TWeakPtr<IBufferedAudioOutput, ESPMode::ThreadSafe>;
+using FSharedBufferedOutputPtr = TSharedPtr<IBufferedAudioOutput, ESPMode::ThreadSafe>;
\ No newline at end of file
diff --git a/Engine/Source/Runtime/AudioMixer/AudioMixer.Build.cs b/Engine/Source/Runtime/AudioMixer/AudioMixer.Build.cs
--- a/Engine/Source/Runtime/AudioMixer/AudioMixer.Build.cs
+++ b/Engine/Source/Runtime/AudioMixer/AudioMixer.Build.cs
@@ -22,7 +22,8 @@ namespace UnrealBuildTool.Rules
 				new string[]
 				{
 					"Core",
-					"CoreUObject"
+					"CoreUObject",
+					"AudioLinkEngine",
 				}
 			);
 
@@ -37,6 +38,7 @@ namespace UnrealBuildTool.Rules
 					"AudioPlatformConfiguration",
 					"SoundFieldRendering",
 					"AudioExtensions",
+					"AudioLinkEngine"
 				}
 			);
 
diff --git a/Engine/Source/Runtime/AudioMixer/Private/AudioMixerSource.cpp b/Engine/Source/Runtime/AudioMixer/Private/AudioMixerSource.cpp
--- a/Engine/Source/Runtime/AudioMixer/Private/AudioMixerSource.cpp
+++ b/Engine/Source/Runtime/AudioMixer/Private/AudioMixerSource.cpp
@@ -397,6 +397,10 @@ namespace Audio
 
 			InitParams.SourceEffectChainId = 0;
 
+			InitParams.SourceBufferListener = WaveInstance->SourceBufferListener;
+			InitParams.bShouldSourceBufferListenerZeroBuffer = WaveInstance->bShouldSourceBufferListenerZeroBuffer;
+
+
 			// Source manager needs to know if this is a vorbis source for rebuilding speaker maps
 			InitParams.bIsVorbis = bIsVorbis;
 
diff --git a/Engine/Source/Runtime/AudioMixer/Private/AudioMixerSourceManager.cpp b/Engine/Source/Runtime/AudioMixer/Private/AudioMixerSourceManager.cpp
--- a/Engine/Source/Runtime/AudioMixer/Private/AudioMixerSourceManager.cpp
+++ b/Engine/Source/Runtime/AudioMixer/Private/AudioMixerSourceManager.cpp
@@ -559,6 +559,12 @@ namespace Audio
 		SourceInfo.LowPassFreq = MAX_FILTER_FREQUENCY;
 		SourceInfo.HighPassFreq = MIN_FILTER_FREQUENCY;
 
+		if (SourceInfo.SourceBufferListener)
+		{
+			SourceInfo.SourceBufferListener->OnSourceReleased(SourceId);
+			SourceInfo.SourceBufferListener.Reset();
+		}
+
 		SourceInfo.ResetModulators(MixerDevice->DeviceID);
 
 		SourceInfo.LowPassFilter.Reset();
@@ -844,6 +850,10 @@ namespace Audio
 				SourceInfo.bUseReverbPlugin = true;
 			}
 
+			// Optional Source Buffer listener.
+			SourceInfo.SourceBufferListener = InitParams.SourceBufferListener;
+			SourceInfo.bShouldSourceBufferListenerZeroBuffer = InitParams.bShouldSourceBufferListenerZeroBuffer;
+
 			// Default all sounds to not consider effect chain tails when playing
 			SourceInfo.bEffectTailsDone = true;
 
@@ -2077,6 +2087,27 @@ namespace Audio
 			return;
 		}
 
+		// If we have Source Buffer Listener
+		if (SourceInfo.SourceBufferListener.IsValid())
+		{
+			// Pack all our state into a single struct.
+			ISourceBufferListener::FOnNewBufferParams Params;
+			Params.SourceId			= SourceId;
+			Params.AudioData		= SourceInfo.PreDistanceAttenuationBuffer.GetData();
+			Params.NumSamples		= SourceInfo.PreDistanceAttenuationBuffer.Num();
+			Params.NumChannels		= SourceInfo.NumInputChannels;
+			Params.SampleRate		= MixerDevice->GetSampleRate();
+
+			// Fire callback.
+			SourceInfo.SourceBufferListener->OnNewBuffer(Params);
+
+			// Optionally, clear the buffer after we've broadcast it. 
+			if (SourceInfo.bShouldSourceBufferListenerZeroBuffer)
+			{
+				FMemory::Memzero(SourceInfo.PreDistanceAttenuationBuffer.GetData(), SourceInfo.PreDistanceAttenuationBuffer.Num() * sizeof(float));
+			}
+		}
+
 		float* PostDistanceAttenBufferPtr = SourceInfo.SourceBuffer.GetData();
 
 		bool bShouldMixInReverb = false;
diff --git a/Engine/Source/Runtime/AudioMixer/Private/AudioMixerSourceManager.h b/Engine/Source/Runtime/AudioMixer/Private/AudioMixerSourceManager.h
--- a/Engine/Source/Runtime/AudioMixer/Private/AudioMixerSourceManager.h
+++ b/Engine/Source/Runtime/AudioMixer/Private/AudioMixerSourceManager.h
@@ -127,6 +127,8 @@ namespace Audio
 
 		FQuartzQuantizedRequestData QuantizedRequestData;
 
+		FSharedISourceBufferListenerPtr SourceBufferListener;		
+
 		FName AudioComponentUserID;
 		uint64 AudioComponentID = 0;
 		bool bIs3D = false;
@@ -140,6 +142,7 @@ namespace Audio
 		bool bIsVorbis = false;
 		bool bIsSoundfield = false;
 		bool bIsSeeking = false;
+		bool bShouldSourceBufferListenerZeroBuffer = false;
 	};
 
 	struct FSourceManagerInitParams
@@ -431,6 +434,9 @@ namespace Audio
 			// Quantization data
 			FQuartzQuantizedCommandHandle QuantizedCommandHandle;
 
+			// Optional Source buffer listener.
+			FSharedISourceBufferListenerPtr SourceBufferListener;
+
 			// State management
 			uint8 bIs3D:1;
 			uint8 bIsCenterChannelOnly:1;
@@ -457,6 +463,7 @@ namespace Audio
 			uint8 bIsBypassingHPF:1;
 			uint8 bHasPreDistanceAttenuationSend:1;
 			uint8 bModFiltersUpdated : 1;
+			uint8 bShouldSourceBufferListenerZeroBuffer : 1;
 
 			// Source format info
 			int32 NumInputChannels;
diff --git a/Engine/Source/Runtime/AudioMixer/Private/AudioMixerSubmix.cpp b/Engine/Source/Runtime/AudioMixer/Private/AudioMixerSubmix.cpp
--- a/Engine/Source/Runtime/AudioMixer/Private/AudioMixerSubmix.cpp
+++ b/Engine/Source/Runtime/AudioMixer/Private/AudioMixerSubmix.cpp
@@ -14,6 +14,7 @@
 #include "Sound/SoundSubmixSend.h"
 #include "Misc/ScopeTryLock.h"
 #include "ProfilingDebugging/CsvProfiler.h"
+#include "AudioLinkLog.h"
 
 // Link to "Audio" profiling category
 CSV_DECLARE_CATEGORY_MODULE_EXTERN(AUDIOMIXERCORE_API, Audio);
@@ -283,7 +284,18 @@ namespace Audio
 					UpdateModulationSettings(VolumeModulator, WetLevelModulator, DryLevelModulator);
 				});
 			}
-
+			if (SoundSubmix->AudioLinkSettings)
+			{
+				if (IAudioLinkFactory* LinkFactory = IAudioLinkFactory::FindFactory(SoundSubmix->AudioLinkSettings->GetFactoryName()))
+				{
+					AudioLinkInstance = LinkFactory->CreateSubmixAudioLink({ SoundSubmix, MixerDevice, SoundSubmix->AudioLinkSettings });
+				}
+				else
+				{
+					UE_LOG(LogAudioLink, Warning, TEXT("Failed to Find AudioLink Factory '%s'"), *SoundSubmix->AudioLinkSettings->GetFactoryName().GetPlainNameString());
+				}
+			}
+			
 			FScopeLock ScopeLock(&EffectChainMutationCriticalSection);
 			{
     			NumSubmixEffects = 0;
diff --git a/Engine/Source/Runtime/AudioMixer/Public/AudioMixerSubmix.h b/Engine/Source/Runtime/AudioMixer/Public/AudioMixerSubmix.h
--- a/Engine/Source/Runtime/AudioMixer/Public/AudioMixerSubmix.h
+++ b/Engine/Source/Runtime/AudioMixer/Public/AudioMixerSubmix.h
@@ -16,6 +16,7 @@
 #include "AudioDynamicParameter.h"
 #include "Stats/Stats.h"
 #include "UObject/WeakObjectPtrTemplates.h"
+#include "IAudioLinkFactory.h"
 
 // The time it takes to process the submix graph. Process submix effects, mix into the submix buffer, etc.
 DECLARE_CYCLE_STAT_EXTERN(TEXT("Submix Graph"), STAT_AudioMixerSubmixes, STATGROUP_AudioMixer, AUDIOMIXER_API);
@@ -644,6 +645,8 @@ namespace Audio
 
 		Audio::FPatchSplitter PatchSplitter;
 
+		TUniquePtr<IAudioLink> AudioLinkInstance;
+
 		friend class FMixerDevice;
 	};
 }
diff --git a/Engine/Source/Runtime/Engine/Classes/Components/AudioComponent.h b/Engine/Source/Runtime/Engine/Classes/Components/AudioComponent.h
--- a/Engine/Source/Runtime/Engine/Classes/Components/AudioComponent.h
+++ b/Engine/Source/Runtime/Engine/Classes/Components/AudioComponent.h
@@ -731,6 +731,12 @@ public:
 	/** Sets whether or not sound instance is virtualized */
 	void SetIsVirtualized(bool bInIsVirtualized);
 
+	/** Sets Source Buffer Listener */
+	void SetSourceBufferListener(const FSharedISourceBufferListenerPtr& InSourceBufferListener, bool bShouldZeroBufferAfter);
+	
+	/** Gets  Source Buffer Listener */
+	const FSharedISourceBufferListenerPtr& GetSourceBufferListener() const { return SourceBufferListener; }
+
 	//~ Begin UObject Interface.
 #if WITH_EDITOR
 	virtual void PostEditChangeProperty(FPropertyChangedEvent& PropertyChangedEvent) override;
@@ -860,6 +866,10 @@ private:
 
 	/** Restore relative transform from auto attachment and optionally detach from parent (regardless of whether it was an auto attachment). */
 	void CancelAutoAttachment(bool bDetachFromParent, const UWorld* MyWorld);
+	
+	/** Source Buffer Listener. */
+	FSharedISourceBufferListenerPtr SourceBufferListener;
+	bool bShouldSourceBufferListenerZeroBuffer = false;
 
 protected:
 
diff --git a/Engine/Source/Runtime/Engine/Classes/Sound/SoundSubmix.h b/Engine/Source/Runtime/Engine/Classes/Sound/SoundSubmix.h
--- a/Engine/Source/Runtime/Engine/Classes/Sound/SoundSubmix.h
+++ b/Engine/Source/Runtime/Engine/Classes/Sound/SoundSubmix.h
@@ -13,6 +13,7 @@
 #include "UObject/Object.h"
 #include "UObject/ObjectMacros.h"
 #include "DSP/SpectrumAnalyzer.h"
+#include "AudioLinkSettingsAbstract.h"
 
 #include "SoundSubmix.generated.h"
 
@@ -323,6 +324,10 @@ public:
 	/** Modulation to apply to the submix Dry Level (in dB)*/
 	UPROPERTY(EditAnywhere, BlueprintReadWrite, Category = Modulation, meta = (DisplayName = "Dry Level Modulation", AudioParam = "Volume"))
 	FSoundModulationDestinationSettings DryLevelModulation;
+	
+	/** Optional Audio Link Settings Object */
+	UPROPERTY(EditAnywhere, BlueprintReadWrite, Category = AudioLink)
+	TObjectPtr<UAudioLinkSettingsAbstract> AudioLinkSettings;
 
 	// Blueprint delegate for when a recorded file is finished exporting.
 	UPROPERTY(BlueprintAssignable)
diff --git a/Engine/Source/Runtime/Engine/Engine.Build.cs b/Engine/Source/Runtime/Engine/Engine.Build.cs
--- a/Engine/Source/Runtime/Engine/Engine.Build.cs
+++ b/Engine/Source/Runtime/Engine/Engine.Build.cs
@@ -97,6 +97,7 @@ public class Engine : ModuleRules
 				"SignalProcessing",
 				"AudioExtensions",
 				"DeveloperSettings",
+				"AudioLinkCore",
 				"CookOnTheFly",
 			}
 		);
diff --git a/Engine/Source/Runtime/Engine/Private/ActiveSound.cpp b/Engine/Source/Runtime/Engine/Private/ActiveSound.cpp
--- a/Engine/Source/Runtime/Engine/Private/ActiveSound.cpp
+++ b/Engine/Source/Runtime/Engine/Private/ActiveSound.cpp
@@ -795,6 +795,10 @@ void FActiveSound::UpdateWaveInstances(TArray<FWaveInstance*> &InWaveInstances,
 				// manually deep copy the QuantizedCommandPtr object itself
 				WaveInstance->QuantizedRequestData->QuantizedCommandPtr = QuantizedRequestData.QuantizedCommandPtr->GetDeepCopyOfDerivedObject();
 			}
+			
+			// each wave instance needs its own copy of the source buffer listener.
+			WaveInstance->SourceBufferListener = SourceBufferListener;
+			WaveInstance->bShouldSourceBufferListenerZeroBuffer = bShouldSourceBufferListenerZeroBuffer;
 		}
 
 		// If the concurrency volume is negative (as set by ConcurrencyManager on creation),
diff --git a/Engine/Source/Runtime/Engine/Private/Components/AudioComponent.cpp b/Engine/Source/Runtime/Engine/Private/Components/AudioComponent.cpp
--- a/Engine/Source/Runtime/Engine/Private/Components/AudioComponent.cpp
+++ b/Engine/Source/Runtime/Engine/Private/Components/AudioComponent.cpp
@@ -764,6 +764,9 @@ void UAudioComponent::PlayInternal(const PlayInternalRequestData& InPlayRequestD
 	// Pass quantization data to the active sound
 	NewActiveSound.QuantizedRequestData = InPlayRequestData.QuantizedRequestData;
 
+	// Pass down any source buffer listener we have
+	NewActiveSound.SetSourceListener(SourceBufferListener, bShouldSourceBufferListenerZeroBuffer);
+
 	NewActiveSound.MaxDistance = MaxDistance;
 
 	Audio::FVolumeFader& Fader = NewActiveSound.ComponentVolumeFader;
@@ -1293,6 +1296,12 @@ void UAudioComponent::SetIsVirtualized(bool bInIsVirtualized)
 	bIsVirtualized = bInIsVirtualized ? 1 : 0;
 }
 
+void UAudioComponent::SetSourceBufferListener(const FSharedISourceBufferListenerPtr& InPtr, bool bShouldZeroBuffer)
+{
+	SourceBufferListener = InPtr;
+	bShouldSourceBufferListenerZeroBuffer = bShouldZeroBuffer;
+}
+
 void UAudioComponent::SetWaveParameter(FName InName, USoundWave* InWave)
 {
 	SetObjectParameter(InName, Cast<UObject>(InWave));
diff --git a/Engine/Source/Runtime/Engine/Public/ActiveSound.h b/Engine/Source/Runtime/Engine/Public/ActiveSound.h
--- a/Engine/Source/Runtime/Engine/Public/ActiveSound.h
+++ b/Engine/Source/Runtime/Engine/Public/ActiveSound.h
@@ -334,6 +334,12 @@ public:
 		AudioDevice = InAudioDevice;
 	}
 
+	void SetSourceListener(FSharedISourceBufferListenerPtr InListener, bool bShouldZeroBuffer)
+	{
+		SourceBufferListener = InListener;
+		bShouldSourceBufferListenerZeroBuffer = bShouldZeroBuffer;
+	}
+
 	int32 GetClosestListenerIndex() const { return ClosestListenerIndex; }
 
 	/** Returns whether or not the active sound can be deleted. */
@@ -630,6 +636,10 @@ public:
 	/** Quantization information */
 	Audio::FQuartzQuantizedRequestData QuantizedRequestData;
 
+	/** Source buffer listener */
+	FSharedISourceBufferListenerPtr SourceBufferListener;
+	bool bShouldSourceBufferListenerZeroBuffer = false;
+
 	/** Cache what volume settings we had last time so we don't have to search again if we didn't move */
 	FInteriorSettings InteriorSettings;
 	TArray<FAudioVolumeSubmixSendSettings> AudioVolumeSubmixSendSettings;
diff --git a/Engine/Source/Runtime/Engine/Public/Audio.h b/Engine/Source/Runtime/Engine/Public/Audio.h
--- a/Engine/Source/Runtime/Engine/Public/Audio.h
+++ b/Engine/Source/Runtime/Engine/Public/Audio.h
@@ -151,6 +151,27 @@ ENGINE_API UClass* GetAudioPluginCustomSettingsClass(EAudioPlugin PluginType);
 /** accessor for our Spatialization enabled CVar. */
 ENGINE_API bool IsSpatializationCVarEnabled();
 
+/**
+ * Interface for listening to source buffers being rendered.
+ */
+class ISourceBufferListener
+{
+public:
+	virtual ~ISourceBufferListener() = default;
+
+	struct FOnNewBufferParams
+	{
+		const float* AudioData = nullptr;
+		int32 SourceId = INDEX_NONE;
+		int32 NumSamples = 0;
+		int32 NumChannels = 0;
+		int32 SampleRate = 0;
+	};
+	virtual void OnNewBuffer(const FOnNewBufferParams&) = 0;
+	virtual void OnSourceReleased(const int32 InSourceId) = 0;
+};
+using FSharedISourceBufferListenerPtr = TSharedPtr<ISourceBufferListener, ESPMode::ThreadSafe>;
+
 /** Bus send types */
 enum class EBusSendType : uint8
 {
@@ -195,6 +216,10 @@ struct ENGINE_API FWaveInstance
 	/** Quantized Request data */
 	TUniquePtr<Audio::FQuartzQuantizedRequestData> QuantizedRequestData;
 
+	/** Source Buffer listener */
+	FSharedISourceBufferListenerPtr SourceBufferListener;
+	bool bShouldSourceBufferListenerZeroBuffer = false;
+
 private:
 
 	/** Current volume */
