AudioLink starvation fixes.

diff --git a/Engine/Source/Runtime/AudioLink/AudioLinkEngine/Private/BufferedListenerBase.cpp b/Engine/Source/Runtime/AudioLink/AudioLinkEngine/Private/BufferedListenerBase.cpp
--- a/Engine/Source/Runtime/AudioLink/AudioLinkEngine/Private/BufferedListenerBase.cpp
+++ b/Engine/Source/Runtime/AudioLink/AudioLinkEngine/Private/BufferedListenerBase.cpp
@@ -12,6 +12,7 @@ FBufferedListenerBase::FBufferedListenerBase(int32 InDefaultCircularBufferSize)
 bool FBufferedListenerBase::PopBuffer(float* InBuffer, int32 InBufferSizeInSamples, int32& OutSamplesWritten)
 {
 	OutSamplesWritten = LocklessCircularBuffer.Pop(InBuffer, InBufferSizeInSamples);
+	NumSilentSamplesRemaining = FMath::Max(0,NumSilentSamplesRemaining-OutSamplesWritten);
 	return bStarted;
 }
 
@@ -33,7 +34,7 @@ void FBufferedListenerBase::SetFormatKnownDelegate(FOnFormatKnown InFormatKnownD
 }
 
 /** AUDIO MIXER THREAD. */
-void FBufferedListenerBase::OnBufferRecieved(const FBufferFormat& InFormat, TArrayView<const float> InBuffer)
+void FBufferedListenerBase::OnBufferReceived(const FBufferFormat& InFormat, TArrayView<const float> InBuffer)
 {
 	// Keep track of if we need to fire the delegate, so we can do it outside of a lock.
 	bool bFireFormatKnownDelegate = false;
@@ -74,8 +75,8 @@ void FBufferedListenerBase::OnBufferRecieved(const FBufferFormat& InFormat, TArr
 	// Push the data into the circular buffer.
 	int32 SamplesPushed = LocklessCircularBuffer.Push(InBuffer.GetData(), InBuffer.Num());
 
-	// Warn of not enough space in circular buffer
-	if (SamplesPushed != InBuffer.Num())
+	// Warn of not enough space in circular buffer, unless we're overwriting silence.
+	if (SamplesPushed < InBuffer.Num() && NumSilentSamplesRemaining == 0)
 	{
 		// Prevent log spam by limiting to 1:100 logs.
 		static const int32 NumLogMessagesToSkip = 100;
@@ -84,7 +85,53 @@ void FBufferedListenerBase::OnBufferRecieved(const FBufferFormat& InFormat, TArr
 	}
 }
 
-void FBufferedListenerBase::Reserve(int32 InNumSamplesToReserve)
+void FBufferedListenerBase::ResetFormat()
 {
-	LocklessCircularBuffer.SetCapacity(InNumSamplesToReserve);
+	FWriteScopeLock WriteLock(FormatKnownRwLock);
+	KnownFormat.Reset();
+}
+
+void FBufferedListenerBase::SetFormat(const FBufferFormat& InFormat)
+{
+	FWriteScopeLock WriteLock(FormatKnownRwLock);
+	KnownFormat = InFormat;
+}
+
+bool FBufferedListenerBase::IsStartedNonAtomic() const
+{
+	return bStarted;
+}
+
+bool FBufferedListenerBase::TrySetStartedFlag()
+{
+	// We expect not to be started.
+	bool bExpected = false;
+	bool bDesired = true;
+	return bStarted.compare_exchange_strong(bExpected, bDesired);
+}
+
+bool FBufferedListenerBase::TryUnsetStartedFlag()
+{
+	// We expect be started.
+	bool bExpected = true;
+	bool bDesired = false;
+	return bStarted.compare_exchange_strong(bExpected, bDesired);
+}
+
+void FBufferedListenerBase::PushSilence(int32 InNumSamplesOfSilence)
+{
+	LocklessCircularBuffer.PushZeros(InNumSamplesOfSilence);
+	NumSilentSamplesRemaining += InNumSamplesOfSilence;
+}
+
+void FBufferedListenerBase::Reserve(int32 InNumSamplesToReserve, int32 InNumSamplesOfSilence)
+{
+	// This Zeros the buffer also so should only be done at the start.
+	LocklessCircularBuffer.SetCapacity(InNumSamplesToReserve);	
+
+	// Optionally add silence into the buffer, this will cause latency.
+	if (InNumSamplesOfSilence > 0)
+	{
+		PushSilence(InNumSamplesOfSilence);
+	}
 }
\ No newline at end of file
diff --git a/Engine/Source/Runtime/AudioLink/AudioLinkEngine/Private/BufferedListenerBase.h b/Engine/Source/Runtime/AudioLink/AudioLinkEngine/Private/BufferedListenerBase.h
--- a/Engine/Source/Runtime/AudioLink/AudioLinkEngine/Private/BufferedListenerBase.h
+++ b/Engine/Source/Runtime/AudioLink/AudioLinkEngine/Private/BufferedListenerBase.h
@@ -20,17 +20,36 @@ protected:
 	//~ Begin IBufferedAudioOutput
 	bool PopBuffer(float* InBuffer, int32 InBufferSizeInSamples, int32& OutSamplesWritten) override;
 	bool GetFormat(IBufferedAudioOutput::FBufferFormat& OutFormat) const override;
-	void Reserve(int32 InNumSamplesToReserve) override;
+	void Reserve(int32 InNumSamplesToReserve, int32 InNumSamplesOfSilence) override;
 	void SetFormatKnownDelegate(FOnFormatKnown InFormatKnownDelegate) override;
+	void SetBufferStreamEndDelegate(FOnBufferStreamEnd) override {}
 	//~ End IBufferedAudioOutput
 
 	//* Common path to receive a new buffer, call from derived classes */
-	void OnBufferRecieved(const FBufferFormat& InFormat, TArrayView<const float> InBuffer);
+	void OnBufferReceived(const FBufferFormat& InFormat, TArrayView<const float> InBuffer);
+
+	//* Reset the format of the buffer */
+	void ResetFormat();
+
+	//* Set the format of the buffer */
+	void SetFormat(const FBufferFormat& InFormat);
+	
+	//* Ask if the started flag has been set. Note this is non-atomic, as it could change during the call */	
+	bool IsStartedNonAtomic() const;
+		
+	//* Attempt to set our state to started. Not this can fail if we're already started. */	
+	bool TrySetStartedFlag(); 
+	
+	//* Attempt to set started to false. This can fail if we're already stopped. */	
+	bool TryUnsetStartedFlag();
+
+private:	
+	void PushSilence(int32 InNumSamplesOfSilence);
 
 	/** Lock-less buffer to hold the data for the single source we're listening to, interleaved. */
 	Audio::TCircularAudioBuffer<float> LocklessCircularBuffer;
 
-	/** Read/Write slim lock protects format known optional below. */
+	/** Read/Write slim lock protects format known optional */
 	mutable FRWLock FormatKnownRwLock;
 
 	/** Optional that holds the buffer format, if (and when) it's known. Protected by r/w Slim-lock */
@@ -39,6 +58,9 @@ protected:
 	/** Delegate that fires when the format it known. Normally on the first buffer received. */
 	FOnFormatKnown OnFormatKnown;
 
-	/** Atomic flag we're been started yet. */
+	/** Count for how much (initial) silence is still remaining in the Circular buffer, for logging purposes */
+	std::atomic<int32> NumSilentSamplesRemaining;
+
+	/** Atomic flag we've been started */
 	std::atomic<bool> bStarted;
 };
diff --git a/Engine/Source/Runtime/AudioLink/AudioLinkEngine/Private/BufferedSourceListener.cpp b/Engine/Source/Runtime/AudioLink/AudioLinkEngine/Private/BufferedSourceListener.cpp
--- a/Engine/Source/Runtime/AudioLink/AudioLinkEngine/Private/BufferedSourceListener.cpp
+++ b/Engine/Source/Runtime/AudioLink/AudioLinkEngine/Private/BufferedSourceListener.cpp
@@ -13,22 +13,14 @@ FBufferedSourceListener::~FBufferedSourceListener()
 	// No need to stop as its only record keeping and nothing to unregister.
 }
 
-bool FBufferedSourceListener::Start(FAudioDevice* InAudioDevice)
+bool FBufferedSourceListener::Start(FAudioDevice*)
 {
-	if (ensure(!bStarted))
-	{
-		bStarted = true;
-		return true;
-	}
-	return false;
+	return TrySetStartedFlag();
 }
 
-void FBufferedSourceListener::Stop(FAudioDevice* InAudioDevice)
+void FBufferedSourceListener::Stop(FAudioDevice*)
 {
-	if (bStarted)
-	{
-		bStarted = false;
-	}
+	ensure(TryUnsetStartedFlag());
 }
 
 /** AUDIO MIXER THREAD. When a source is finished and returned to the pool, this call will be called. */
@@ -36,18 +28,28 @@ void FBufferedSourceListener::OnSourceReleased(const int32 InSourceId)
 {	
 	if (InSourceId == CurrentSourceId)
 	{
+		// Fire delegate to any interested parties.
+		IBufferedAudioOutput::FBufferStreamEnd Params;
+		Params.Id = CurrentSourceId;
+		OnBufferStreamEndDelegate.ExecuteIfBound(Params);
+
+		// Reset our source id.
 		CurrentSourceId = INDEX_NONE;
-		
+
 		// Reset the format, now our source we were listening to has stopped.
-		FWriteScopeLock WriteLock(FormatKnownRwLock);
-		KnownFormat.Reset();
+		ResetFormat();
 	}
 }
 
+void FBufferedSourceListener::SetBufferStreamEndDelegate(FOnBufferStreamEnd InBufferStreamEndDelegate)
+{
+	OnBufferStreamEndDelegate = InBufferStreamEndDelegate;
+}
+
 /** AUDIO MIXER THREAD. New Audio buffers from the active sources enter here. */
 void FBufferedSourceListener::OnNewBuffer(const ISourceBufferListener::FOnNewBufferParams& InParams)
-{	
-	if (bStarted) 
+{
+	if (IsStartedNonAtomic()) 
 	{
 		// Make sure caller is sane.
 		check(InParams.AudioData);
@@ -73,6 +75,6 @@ void FBufferedSourceListener::OnNewBuffer(const ISourceBufferListener::FOnNewBuf
 		NewFormat.NumChannels = InParams.NumChannels;
 		NewFormat.NumSamplesPerBlock = InParams.NumSamples;
 		NewFormat.NumSamplesPerSec = InParams.SampleRate;
-		OnBufferRecieved(NewFormat, MakeArrayView(InParams.AudioData, InParams.NumSamples));
+		OnBufferReceived(NewFormat, MakeArrayView(InParams.AudioData, InParams.NumSamples));
 	}
 }
diff --git a/Engine/Source/Runtime/AudioLink/AudioLinkEngine/Private/BufferedSourceListener.h b/Engine/Source/Runtime/AudioLink/AudioLinkEngine/Private/BufferedSourceListener.h
--- a/Engine/Source/Runtime/AudioLink/AudioLinkEngine/Private/BufferedSourceListener.h
+++ b/Engine/Source/Runtime/AudioLink/AudioLinkEngine/Private/BufferedSourceListener.h
@@ -32,5 +32,10 @@ private:
 	void OnSourceReleased(const int32 SourceId) override;
 	//~ End ISourceBufferListener
 
+	//~ Begin IBufferedAudioOutput
+	void SetBufferStreamEndDelegate(FOnBufferStreamEnd InBufferStreamEndDelegate) override;
+	//~ End IBufferedAudioOutput
+
 	std::atomic<int32> CurrentSourceId	= INDEX_NONE;		// r/w AudioMixer thread.	
+	FOnBufferStreamEnd OnBufferStreamEndDelegate;
 };
diff --git a/Engine/Source/Runtime/AudioLink/AudioLinkEngine/Private/BufferedSubmixListener.cpp b/Engine/Source/Runtime/AudioLink/AudioLinkEngine/Private/BufferedSubmixListener.cpp
--- a/Engine/Source/Runtime/AudioLink/AudioLinkEngine/Private/BufferedSubmixListener.cpp
+++ b/Engine/Source/Runtime/AudioLink/AudioLinkEngine/Private/BufferedSubmixListener.cpp
@@ -2,31 +2,48 @@
 
 #include "BufferedSubmixListener.h"
 
+namespace BufferedSubmixListenerPrivate
+{
+	static Audio::FDeviceId InvalidAudioDeviceId = static_cast<Audio::FDeviceId>(INDEX_NONE);
+}
+
 /** Buffered Submix Listener. */
 FBufferedSubmixListener::FBufferedSubmixListener(int32 InDefaultCircularBufferSize, bool bInZeroInputBuffer)
 	: FBufferedListenerBase{ InDefaultCircularBufferSize }
+	, DeviceId{ BufferedSubmixListenerPrivate::InvalidAudioDeviceId }
 	, bZeroInputBuffer{ bInZeroInputBuffer }
 {
 }
 
 FBufferedSubmixListener::~FBufferedSubmixListener()
 {
-	//Stop(); // FIXME! Virtual call in destructor.	
-	//StopInternal();
-	// Perhaps we should assert that we're still registered.
+	// We should have been unregistered.
+	check(DeviceId == BufferedSubmixListenerPrivate::InvalidAudioDeviceId);
+	check(!IsStartedNonAtomic());
+}
+
+void FBufferedSubmixListener::RegisterWithAudioDevice(FAudioDevice* InDevice)
+{
+	DeviceId = InDevice->DeviceID;
+	InDevice->RegisterSubmixBufferListener(this);
+}
+
+void FBufferedSubmixListener::UnregsiterWithAudioDevice(FAudioDevice* InDevice)
+{
+	if (ensure(DeviceId != BufferedSubmixListenerPrivate::InvalidAudioDeviceId))
+	{
+		DeviceId = BufferedSubmixListenerPrivate::InvalidAudioDeviceId;
+		InDevice->UnregisterSubmixBufferListener(this);
+	}
 }
 
 bool FBufferedSubmixListener::Start(FAudioDevice* InAudioDevice)
 {
-	if (ensure(!bStarted))
-	{		
-		if (ensure(InAudioDevice))
+	if (ensure(InAudioDevice))
+	{
+		if (TrySetStartedFlag())
 		{
-			// Inputs look valid, we can start.
-			DeviceId = InAudioDevice->DeviceID;
-			InAudioDevice->RegisterSubmixBufferListener(this);
-			bStarted = true;
-			return true;
+			RegisterWithAudioDevice(InAudioDevice);
 		}
 	}
 	return false;
@@ -34,19 +51,13 @@ bool FBufferedSubmixListener::Start(FAudioDevice* InAudioDevice)
 
 void FBufferedSubmixListener::Stop(FAudioDevice* InAudioDevice)
 {
-	StopInternal(InAudioDevice);
-}
-
-void FBufferedSubmixListener::StopInternal(FAudioDevice* InAudioDevice)
-{
-	if (bStarted)
+	if (ensure(InAudioDevice))
 	{
-		if (ensure(InAudioDevice))
+		if (ensure(InAudioDevice->DeviceID == DeviceId))
 		{
-			if (ensure(InAudioDevice->DeviceID == DeviceId))
+			if (TryUnsetStartedFlag())
 			{
-				InAudioDevice->UnregisterSubmixBufferListener(this);
-				DeviceId = bStarted = false;
+				UnregsiterWithAudioDevice(InAudioDevice);
 			}
 		}
 	}
@@ -54,14 +65,14 @@ void FBufferedSubmixListener::StopInternal(FAudioDevice* InAudioDevice)
 
 void FBufferedSubmixListener::OnNewSubmixBuffer(const USoundSubmix* OwningSubmix, float* AudioData, int32 InNumSamples, int32 InNumChannels, const int32 InSampleRate, double)
 {
-	if(bStarted)
+	if (IsStartedNonAtomic())
 	{
 		// Call to base class to handle.
 		FBufferFormat NewFormat;
 		NewFormat.NumChannels = InNumChannels;
 		NewFormat.NumSamplesPerBlock = InNumSamples;
 		NewFormat.NumSamplesPerSec = InSampleRate;
-		OnBufferRecieved(NewFormat, MakeArrayView(AudioData, InNumSamples));
+		OnBufferReceived(NewFormat, MakeArrayView(AudioData, InNumSamples));
 
 		// Optionally, zero the buffer if we're asked to. This in the case where we're running both Unreal+Consumer renderers at once.
 		// NOTE: this is dangerous as there's a chance we're not the only listener registered on this Submix. And will cause
diff --git a/Engine/Source/Runtime/AudioLink/AudioLinkEngine/Private/BufferedSubmixListener.h b/Engine/Source/Runtime/AudioLink/AudioLinkEngine/Private/BufferedSubmixListener.h
--- a/Engine/Source/Runtime/AudioLink/AudioLinkEngine/Private/BufferedSubmixListener.h
+++ b/Engine/Source/Runtime/AudioLink/AudioLinkEngine/Private/BufferedSubmixListener.h
@@ -29,8 +29,9 @@ public:
 	void Stop(FAudioDevice* InDevice) override;
 
 private:
-	void StopInternal(FAudioDevice* InDevice);
-
+	void RegisterWithAudioDevice(FAudioDevice* InDevice);
+	void UnregsiterWithAudioDevice(FAudioDevice* InDevice);
+	
 	//~ Begin ISubmixBufferListener
 	void OnNewSubmixBuffer(const USoundSubmix* OwningSubmix, float* AudioData, int32 InNumSamples, int32 InNumChannels, const int32 InSampleRate, double) override;
 	//~ End ISubmixBufferListener
diff --git a/Engine/Source/Runtime/AudioLink/AudioLinkEngine/Public/IBufferedAudioOutput.h b/Engine/Source/Runtime/AudioLink/AudioLinkEngine/Public/IBufferedAudioOutput.h
--- a/Engine/Source/Runtime/AudioLink/AudioLinkEngine/Public/IBufferedAudioOutput.h
+++ b/Engine/Source/Runtime/AudioLink/AudioLinkEngine/Public/IBufferedAudioOutput.h
@@ -36,6 +36,18 @@ public:
 	DECLARE_DELEGATE_OneParam(FOnFormatKnown, FBufferFormat);
 	virtual void SetFormatKnownDelegate(FOnFormatKnown InFormatKnownDelegate) = 0;
 
+	/**
+	 * Stream of buffers ended.
+	 */
+	struct FBufferStreamEnd
+	{
+		int32 Id = INDEX_NONE;
+	};
+
+	// Delegates.
+	DECLARE_DELEGATE_OneParam(FOnBufferStreamEnd, FBufferStreamEnd);
+	virtual void SetBufferStreamEndDelegate(FOnBufferStreamEnd InBufferStreamEndDelegate) = 0;
+
 	virtual bool Start(FAudioDevice* InAudioDevice) = 0;
 	virtual void Stop(FAudioDevice* InAudioDevice) = 0;
 
@@ -62,7 +74,7 @@ public:
 	 *
 	 * @param InNumSamplesToReserve Reserve this number of samples.
 	 */
-	virtual void Reserve(int32 InNumSamplesToReserve) = 0;
+	virtual void Reserve(int32 InNumSamplesToReserve, int32 InNumSamplesOfSilence = 0) = 0;
 };
 
 using FWeakBufferedOutputPtr = TWeakPtr<IBufferedAudioOutput, ESPMode::ThreadSafe>;
