Adding module API to AudioLink.h

diff --git a/Engine/Source/Runtime/AudioLink/AudioLinkEngine/Public/IAudioLink.h b/Engine/Source/Runtime/AudioLink/AudioLinkEngine/Public/IAudioLink.h
--- a/Engine/Source/Runtime/AudioLink/AudioLinkEngine/Public/IAudioLink.h
+++ b/Engine/Source/Runtime/AudioLink/AudioLinkEngine/Public/IAudioLink.h
@@ -6,7 +6,7 @@
   *	Abstract interface for AudioLink instances. 
   *	Purely opaque.
   */
-class IAudioLink
+class AUDIOLINKENGINE_API IAudioLink
 {
 protected:
 	IAudioLink();
