AudioLink Platform Impl.

diff --git a/Engine/Source/Runtime/AudioLink/AudioMixerPlatformAudioLink/AudioMixerPlatformAudioLink.Build.cs b/Engine/Source/Runtime/AudioLink/AudioMixerPlatformAudioLink/AudioMixerPlatformAudioLink.Build.cs
new file mode 100644
--- /dev/null
+++ b/Engine/Source/Runtime/AudioLink/AudioMixerPlatformAudioLink/AudioMixerPlatformAudioLink.Build.cs
@@ -0,0 +1,35 @@
+// Copyright Epic Games, Inc. All Rights Reserved.
+
+using UnrealBuildTool;
+
+public class AudioMixerPlatformAudioLink: ModuleRules
+{
+	public AudioMixerPlatformAudioLink(ReadOnlyTargetRules Target) : base(Target)
+	{
+        PrivateIncludePathModuleNames.Add("TargetPlatform");
+		PublicIncludePaths.Add("Runtime/AudioMixer/Public");
+		PrivateIncludePaths.Add("Runtime/AudioMixer/Private");
+
+
+		PrivateDependencyModuleNames.AddRange(
+			new string[] {
+				"Core",
+				"CoreUObject",
+				"Engine",
+				"AudioMixer",
+				"AudioMixerCore"
+		});
+					
+		if (Target.bCompileAgainstEngine) 
+		{
+			//PrivateDependencyModuleNames.Add("XMA2");
+			PrivateDependencyModuleNames.Add("BinkAudioDecoder");
+
+			AddEngineThirdPartyPrivateStaticDependencies(Target,
+				"UEOgg",
+				"Vorbis",
+				"VorbisFile"
+				);
+		}
+	}
+}
diff --git a/Engine/Source/Runtime/AudioLink/AudioMixerPlatformAudioLink/Private/AudioMixerPlatformAudioLink.cpp b/Engine/Source/Runtime/AudioLink/AudioMixerPlatformAudioLink/Private/AudioMixerPlatformAudioLink.cpp
new file mode 100644
--- /dev/null
+++ b/Engine/Source/Runtime/AudioLink/AudioMixerPlatformAudioLink/Private/AudioMixerPlatformAudioLink.cpp
@@ -0,0 +1,190 @@
+// Copyright Epic Games, Inc. All Rights Reserved.
+
+#include "AudioMixerPlatformAudioLink.h"
+#include "Modules/ModuleManager.h"
+#include "AudioMixer.h"
+#include "AudioMixerDevice.h"
+#include "CoreGlobals.h"
+#include "Misc/ConfigCacheIni.h"
+
+#include "Misc/CoreMisc.h"
+#include "Interfaces/ITargetPlatform.h"
+#include "Interfaces/ITargetPlatformManagerModule.h"
+
+namespace AudioMixerPlatformAudioLinkPrivate
+{	
+	using namespace Audio;
+	static const FAudioPlatformDeviceInfo& GetPlatformInfo()
+	{
+		static const FAudioPlatformDeviceInfo DeviceInfo = []()
+		{
+			// For now just hard-code a virtual audio platform, but this will ultimately match an external endpoint via audio link.
+			FAudioPlatformDeviceInfo Info;
+			Info.Name = TEXT("AudioLink Virtual Platform");
+			Info.DeviceId = Info.Name;
+			Info.SampleRate	= 48000;
+			Info.NumChannels = 8;
+			Info.bIsSystemDefault = true;
+			Info.Format = EAudioMixerStreamDataFormat::Float;
+			Info.OutputChannelArray.Add(EAudioMixerChannel::FrontLeft);
+			Info.OutputChannelArray.Add(EAudioMixerChannel::FrontRight);
+			Info.OutputChannelArray.Add(EAudioMixerChannel::FrontCenter);
+			Info.OutputChannelArray.Add(EAudioMixerChannel::LowFrequency);
+			Info.OutputChannelArray.Add(EAudioMixerChannel::SideLeft);
+			Info.OutputChannelArray.Add(EAudioMixerChannel::SideRight);
+			Info.OutputChannelArray.Add(EAudioMixerChannel::BackLeft);
+			Info.OutputChannelArray.Add(EAudioMixerChannel::BackRight);
+
+			return Info;
+		}();
+		return DeviceInfo;
+	}
+}
+
+namespace Audio
+{		
+	FAudioMixerPlatformAudioLink::FAudioMixerPlatformAudioLink()
+	{}
+
+	bool FAudioMixerPlatformAudioLink::InitializeHardware()
+	{
+		if (IAudioMixer::ShouldRecycleThreads())
+		{
+			// Pre-create the null render device thread on XAudio2, so we can simple wake it up when we need it.
+			// Give it nothing to do, with a slow tick as the default, but ask it to wait for a signal to wake up.
+			CreateNullDeviceThread([] {}, 1.0f, true);
+		}
+
+		bInitialized = true;
+		return true;
+	}
+
+	bool FAudioMixerPlatformAudioLink::TeardownHardware()
+	{		
+		StopAudioStream();
+		CloseAudioStream();
+		return true;
+	}
+
+	bool FAudioMixerPlatformAudioLink::IsInitialized() const
+	{
+		return bInitialized;
+	}
+
+	bool FAudioMixerPlatformAudioLink::GetNumOutputDevices(uint32& OutNumOutputDevices)
+	{
+		OutNumOutputDevices = 1;
+		return true;
+	}
+	
+	bool FAudioMixerPlatformAudioLink::GetOutputDeviceInfo(const uint32 InDeviceIndex, FAudioPlatformDeviceInfo& OutInfo)
+	{	
+		using namespace AudioMixerPlatformAudioLinkPrivate;
+		
+		OutInfo = GetPlatformInfo();
+		return true;
+	}
+
+	bool FAudioMixerPlatformAudioLink::GetDefaultOutputDeviceIndex(uint32& OutDefaultDeviceIndex) const
+	{
+		// It's not possible to know what index the default audio device is.
+		OutDefaultDeviceIndex = AUDIO_MIXER_DEFAULT_DEVICE_INDEX;
+		return true;
+	}
+
+	bool FAudioMixerPlatformAudioLink::OpenAudioStream(const FAudioMixerOpenStreamParams& Params)
+	{
+		if (!bInitialized || AudioStreamInfo.StreamState != EAudioOutputStreamState::Closed)
+		{
+			return false;
+		}
+			
+		AudioStreamInfo.Reset();
+		if(!GetOutputDeviceInfo(Params.OutputDeviceIndex, AudioStreamInfo.DeviceInfo))
+		{
+			return false;
+		}
+
+		OpenStreamParams = Params;
+
+		AudioStreamInfo.AudioMixer = Params.AudioMixer;
+		AudioStreamInfo.NumBuffers = Params.NumBuffers;
+		AudioStreamInfo.NumOutputFrames = Params.NumFrames;
+		AudioStreamInfo.StreamState = EAudioOutputStreamState::Open;
+
+		return true;
+	}
+
+	bool FAudioMixerPlatformAudioLink::CloseAudioStream()
+	{
+		if (AudioStreamInfo.StreamState == EAudioOutputStreamState::Closed)
+		{
+			return false;
+		}
+
+		if (!StopAudioStream())
+		{
+			return false;
+		}
+
+		AudioStreamInfo.StreamState = EAudioOutputStreamState::Closed;
+		return true;
+	}
+
+	bool FAudioMixerPlatformAudioLink::StartAudioStream()
+	{
+		if (!bInitialized || (AudioStreamInfo.StreamState != EAudioOutputStreamState::Open && AudioStreamInfo.StreamState != EAudioOutputStreamState::Stopped))
+		{
+			return false;
+		}
+
+		// Start generating audio
+		BeginGeneratingAudio();
+
+		StartRunningNullDevice();
+
+		AudioStreamInfo.StreamState = EAudioOutputStreamState::Running;
+
+		return true;
+	}
+
+	bool FAudioMixerPlatformAudioLink::StopAudioStream()
+	{
+		if (AudioStreamInfo.StreamState != EAudioOutputStreamState::Stopped && AudioStreamInfo.StreamState != EAudioOutputStreamState::Closed)
+		{
+			if(bIsUsingNullDevice)
+			{
+				StopRunningNullDevice();
+			}
+			
+			if (AudioStreamInfo.StreamState == EAudioOutputStreamState::Running)
+			{
+				StopGeneratingAudio();
+			}
+
+			check(AudioStreamInfo.StreamState == EAudioOutputStreamState::Stopped);
+		}
+
+		return true;
+	}
+
+	FAudioPlatformDeviceInfo FAudioMixerPlatformAudioLink::GetPlatformDeviceInfo() const
+	{
+		return AudioStreamInfo.DeviceInfo;
+	}
+
+	FString FAudioMixerPlatformAudioLink::GetDefaultDeviceName()
+	{
+		using namespace AudioMixerPlatformAudioLinkPrivate;
+		return GetPlatformInfo().Name;
+	}
+	
+	FAudioPlatformSettings FAudioMixerPlatformAudioLink::GetPlatformSettings() const
+	{
+#if WITH_ENGINE
+		return FAudioPlatformSettings::GetPlatformSettings(FPlatformProperties::GetRuntimeSettingsClassName());
+#else
+		return FAudioPlatformSettings();
+#endif // WITH_ENGINE
+	}
+}
diff --git a/Engine/Source/Runtime/AudioLink/AudioMixerPlatformAudioLink/Private/AudioMixerPlatformAudioLink.h b/Engine/Source/Runtime/AudioLink/AudioMixerPlatformAudioLink/Private/AudioMixerPlatformAudioLink.h
new file mode 100644
--- /dev/null
+++ b/Engine/Source/Runtime/AudioLink/AudioMixerPlatformAudioLink/Private/AudioMixerPlatformAudioLink.h
@@ -0,0 +1,38 @@
+// Copyright Epic Games, Inc. All Rights Reserved.
+
+#pragma once
+
+#include "AudioMixer.h"
+
+namespace Audio
+{
+	class FAudioMixerPlatformAudioLink : public IAudioMixerPlatformInterface
+	{
+	public:
+
+		FAudioMixerPlatformAudioLink();
+		virtual ~FAudioMixerPlatformAudioLink() = default;
+
+		//~ Begin IAudioMixerPlatformInterface Interface
+		FString GetPlatformApi() const override { return TEXT("AudioLink"); }
+		bool InitializeHardware() override;
+		bool TeardownHardware() override;
+		bool IsInitialized() const override;
+		bool GetNumOutputDevices(uint32& OutNumOutputDevices) override;
+		bool GetOutputDeviceInfo(const uint32 InDeviceIndex, FAudioPlatformDeviceInfo& OutInfo) override;
+		bool GetDefaultOutputDeviceIndex(uint32& OutDefaultDeviceIndex) const override;
+		virtual bool OpenAudioStream(const FAudioMixerOpenStreamParams& Params) override;
+		bool CloseAudioStream() override;
+		bool StartAudioStream() override;
+		bool StopAudioStream() override;
+		FAudioPlatformDeviceInfo GetPlatformDeviceInfo() const override;
+		FString GetDefaultDeviceName() override;
+		FAudioPlatformSettings GetPlatformSettings() const override;
+		//~ End IAudioMixerPlatformInterface Interface
+
+	private:
+		bool bSuspended = false;
+		bool bInitialized = false;
+	};
+}
+
diff --git a/Engine/Source/Runtime/AudioLink/AudioMixerPlatformAudioLink/Private/AudioMixerPlatformAudioLinkModule.cpp b/Engine/Source/Runtime/AudioLink/AudioMixerPlatformAudioLink/Private/AudioMixerPlatformAudioLinkModule.cpp
new file mode 100644
--- /dev/null
+++ b/Engine/Source/Runtime/AudioLink/AudioMixerPlatformAudioLink/Private/AudioMixerPlatformAudioLinkModule.cpp
@@ -0,0 +1,31 @@
+// Copyright Epic Games, Inc. All Rights Reserved.
+
+#include "CoreMinimal.h"
+#include "AudioMixer.h"
+#include "AudioMixerDevice.h"
+#include "AudioMixerPlatformAudioLink.h"
+
+class FAudioMixerPlatformAudioLinkModule : public IAudioDeviceModule
+{
+public:
+	virtual bool IsAudioMixerModule() const override { return true; }
+
+	virtual void StartupModule() override
+	{
+		IAudioDeviceModule::StartupModule();
+
+		FModuleManager::Get().LoadModuleChecked(TEXT("AudioMixerCore"));
+	}
+
+	virtual FAudioDevice* CreateAudioDevice() override
+	{
+		return new Audio::FMixerDevice(CreateAudioMixerPlatformInterface());
+	}
+
+	virtual Audio::IAudioMixerPlatformInterface* CreateAudioMixerPlatformInterface() override
+	{
+		return new Audio::FAudioMixerPlatformAudioLink();
+	}
+};
+
+IMPLEMENT_MODULE(FAudioMixerPlatformAudioLinkModule, AudioMixerPlatformAudioLink);
diff --git a/Engine/Source/Runtime/Launch/Launch.Build.cs b/Engine/Source/Runtime/Launch/Launch.Build.cs
--- a/Engine/Source/Runtime/Launch/Launch.Build.cs
+++ b/Engine/Source/Runtime/Launch/Launch.Build.cs
@@ -77,6 +77,10 @@ public class Launch : ModuleRules
 				});
 			}
 
+			DynamicallyLoadedModuleNames.AddRange(new string[] {
+					"AudioMixerPlatformAudioLink",
+				});
+
 			if (Target.Platform.IsInGroup(UnrealPlatformGroup.Windows))
 			{
 				DynamicallyLoadedModuleNames.AddRange(new string[] {
