AudioLink Engine mods.

diff --git a/Engine/Source/Runtime/AudioLink/AudioLinkCore/Public/AudioLinkSettingsAbstract.h b/Engine/Source/Runtime/AudioLink/AudioLinkCore/Public/AudioLinkSettingsAbstract.h
--- a/Engine/Source/Runtime/AudioLink/AudioLinkCore/Public/AudioLinkSettingsAbstract.h
+++ b/Engine/Source/Runtime/AudioLink/AudioLinkCore/Public/AudioLinkSettingsAbstract.h
@@ -6,6 +6,8 @@
 #include "Templates/SharedPointer.h"
 #include "AudioLinkSettingsAbstract.generated.h"
 
+class UAudioLinkSettingsAbstract;
+
 /**
   * This interface should be used to provide a non-uclass version of the data described in
   * your implementation of UAudioLinkSettingsAbstract
@@ -14,6 +16,11 @@ class AUDIOLINKCORE_API IAudioLinkSettingsProxy
 {
 public:
 	virtual ~IAudioLinkSettingsProxy() = default;
+protected:
+#if WITH_EDITOR
+	friend class UAudioLinkSettingsAbstract;
+	virtual void RefreshFromSettings(UAudioLinkSettingsAbstract* InSettings, FPropertyChangedEvent& InPropertyChangedEvent) = 0;
+#endif //WITH_EDITOR
 };
 
 /**
@@ -46,6 +53,13 @@ public:
 	}
 
 protected:
+#if WITH_EDITOR	
+	void PostEditChangeProperty(struct FPropertyChangedEvent& PropertyChangedEvent)
+	{
+		GetProxy()->RefreshFromSettings(this, PropertyChangedEvent);
+	}
+#endif //WITH_EDITOR
+
 	virtual FSharedSettingsProxyPtr MakeProxy() const PURE_VIRTUAL(UAudioLinkSettingsAbstract::MakeProxy, return {};);
 	mutable FSharedSettingsProxyPtr ProxyInstance;
 };
diff --git a/Engine/Source/Runtime/AudioLink/AudioLinkEngine/Private/BufferedListenerBase.cpp b/Engine/Source/Runtime/AudioLink/AudioLinkEngine/Private/BufferedListenerBase.cpp
--- a/Engine/Source/Runtime/AudioLink/AudioLinkEngine/Private/BufferedListenerBase.cpp
+++ b/Engine/Source/Runtime/AudioLink/AudioLinkEngine/Private/BufferedListenerBase.cpp
@@ -81,7 +81,7 @@ void FBufferedListenerBase::OnBufferReceived(const FBufferFormat& InFormat, TArr
 		// Prevent log spam by limiting to 1:100 logs.
 		static const int32 NumLogMessagesToSkip = 100;
 		static int32 LogPacifier = 0;
-		UE_CLOG(LogPacifier++ % NumLogMessagesToSkip == 0, LogAudioLink, Warning, TEXT("Overflow by '%d' Samples in Buffer Listener"), InBuffer.Num() - SamplesPushed);
+		UE_CLOG(LogPacifier++ % NumLogMessagesToSkip == 0, LogAudioLink, Log, TEXT("Overflow by '%d' Samples in Buffer Listener"), InBuffer.Num() - SamplesPushed);
 	}
 }
 
diff --git a/Engine/Source/Runtime/AudioLink/AudioLinkEngine/Private/IAudioLinkFactory.cpp b/Engine/Source/Runtime/AudioLink/AudioLinkEngine/Private/IAudioLinkFactory.cpp
--- a/Engine/Source/Runtime/AudioLink/AudioLinkEngine/Private/IAudioLinkFactory.cpp
+++ b/Engine/Source/Runtime/AudioLink/AudioLinkEngine/Private/IAudioLinkFactory.cpp
@@ -18,22 +18,51 @@ IAudioLinkFactory::~IAudioLinkFactory()
 }
 
 FSharedBufferedOutputPtr IAudioLinkFactory::CreateSourceBufferListener(const FSourceBufferListenerCreateParams& InSourceCreateParams)
-{
-	check(IsInGameThread()); 
+{	
 	auto SourceBufferListenerSP = MakeShared<FBufferedSourceListener, ESPMode::ThreadSafe>(InSourceCreateParams.SizeOfBufferInFrames);
-	InSourceCreateParams.AudioComponent->SetSourceBufferListener(SourceBufferListenerSP, InSourceCreateParams.bShouldZeroBuffer);
+	if (!InSourceCreateParams.AudioComponent.IsExplicitlyNull())
+	{
+		check(IsInGameThread());
+		InSourceCreateParams.AudioComponent->SetSourceBufferListener(SourceBufferListenerSP, InSourceCreateParams.bShouldZeroBuffer);
+	}
 	return SourceBufferListenerSP;
 }
 
-FSharedBufferedOutputPtr IAudioLinkFactory::CreateSubmixBufferListener(const FSubmixBufferListenerCreateParams& InSubmixCreateParams)
+FSharedBufferedOutputPtr IAudioLinkFactory::CreatePushableBufferListener(const FPushedBufferListenerCreateParams& InPushableCreateParams)
 {
-	return MakeShared<FBufferedSubmixListener, ESPMode::ThreadSafe>(InSubmixCreateParams.SizeOfBufferInFrames, InSubmixCreateParams.bShouldZeroBuffer);	
+	// Add push functionality to the source buffer listener with simple wrapper
+	struct FPushableSourceBufferListener : FBufferedSourceListener, IPushableAudioOutput
+	{
+		using FBufferedSourceListener::FBufferedSourceListener;
+
+		IPushableAudioOutput* GetPushableInterface() override { return this; }
+		const IPushableAudioOutput* GetPushableInterface() const { return this; }
+
+		void PushNewBuffer(const IPushableAudioOutput::FOnNewBufferParams& InNewBuffer) override
+		{
+			ISourceBufferListener::FOnNewBufferParams Params;
+			Params.AudioData = InNewBuffer.AudioData;
+			Params.NumChannels = InNewBuffer.NumChannels;
+			Params.NumSamples = InNewBuffer.NumSamples;
+			Params.SourceId = InNewBuffer.Id;
+			Params.SampleRate = InNewBuffer.SampleRate;
+			static_cast<ISourceBufferListener*>(this)->OnNewBuffer(Params);
+		}
+
+		void LastBuffer(int32 InId) override
+		{
+			static_cast<ISourceBufferListener*>(this)->OnSourceReleased(InId);
+		}
+	};
+
+	auto SourceBufferListenerSP = MakeShared<FPushableSourceBufferListener, ESPMode::ThreadSafe>(InPushableCreateParams.SizeOfBufferInFrames);
+	
+	return SourceBufferListenerSP;
 }
 
-FName IAudioLinkFactory::GetModularFeatureName()
+FSharedBufferedOutputPtr IAudioLinkFactory::CreateSubmixBufferListener(const FSubmixBufferListenerCreateParams& InSubmixCreateParams)
 {
-	static FName FeatureName = FName(TEXT("AudioLink Factory"));
-	return FeatureName;
+	return MakeShared<FBufferedSubmixListener, ESPMode::ThreadSafe>(InSubmixCreateParams.SizeOfBufferInFrames, InSubmixCreateParams.bShouldZeroBuffer);	
 }
 
 TArray<IAudioLinkFactory*> IAudioLinkFactory::GetAllRegisteredFactories()
diff --git a/Engine/Source/Runtime/AudioLink/AudioLinkEngine/Public/IAudioLink.h b/Engine/Source/Runtime/AudioLink/AudioLinkEngine/Public/IAudioLink.h
--- a/Engine/Source/Runtime/AudioLink/AudioLinkEngine/Public/IAudioLink.h
+++ b/Engine/Source/Runtime/AudioLink/AudioLinkEngine/Public/IAudioLink.h
@@ -3,6 +3,8 @@
 #pragma once
 
 #include "HAL/Platform.h"
+#include "Math/Transform.h"
+#include "Delegates/DelegateCombinations.h"
 
 /** IAudioLink
   *	Abstract interface for AudioLink instances. 
@@ -16,3 +18,45 @@ public:
 	virtual ~IAudioLink() = default;
 };
 
+/** IAudioLinkSourcePushed
+  *	Where the owning object needs to push it's state
+  */
+class AUDIOLINKENGINE_API IAudioLinkSourcePushed : public IAudioLink
+{
+protected:
+	IAudioLinkSourcePushed() = default;
+public:
+	virtual ~IAudioLinkSourcePushed() = default;
+
+	struct FOnUpdateWorldStateParams
+	{
+		FTransform	WorldTransform;
+	};
+	virtual void OnUpdateWorldState(const FOnUpdateWorldStateParams&) = 0;
+
+	struct FOnNewBufferParams
+	{
+		TArrayView<float> Buffer;
+		int32 SourceId = INDEX_NONE;
+	};
+	virtual void OnNewBuffer(const FOnNewBufferParams&) = 0;
+
+	virtual void OnSourceReleased(const int32 SourceId) = 0;	
+};
+
+/** IAudioLinkSynchronizer
+  *	Provides delegates for hooking and external AudioLinks synchronization callbacks.
+  */
+class IAudioLinkSynchronizer
+{
+protected:
+	IAudioLinkSynchronizer() = default;
+public:
+	virtual ~IAudioLinkSynchronizer() = default;
+
+	DECLARE_DELEGATE(FOnBeginRender);
+	virtual FOnBeginRender& GetOnBeginRender() = 0;
+
+	DECLARE_DELEGATE(FOnEndRender);	
+	virtual FOnEndRender& GetOnEndRender() = 0;
+};
\ No newline at end of file
diff --git a/Engine/Source/Runtime/AudioLink/AudioLinkEngine/Public/IAudioLinkFactory.h b/Engine/Source/Runtime/AudioLink/AudioLinkEngine/Public/IAudioLinkFactory.h
--- a/Engine/Source/Runtime/AudioLink/AudioLinkEngine/Public/IAudioLinkFactory.h
+++ b/Engine/Source/Runtime/AudioLink/AudioLinkEngine/Public/IAudioLinkFactory.h
@@ -48,6 +48,13 @@ public:
 	};
 	virtual FSharedBufferedOutputPtr CreateSourceBufferListener(const FSourceBufferListenerCreateParams&);
 
+	struct FPushedBufferListenerCreateParams
+	{
+		int32 SizeOfBufferInFrames = INDEX_NONE;
+		bool bShouldZeroBuffer = false;
+	};
+	virtual FSharedBufferedOutputPtr CreatePushableBufferListener(const FPushedBufferListenerCreateParams&);
+
 	struct FSubmixBufferListenerCreateParams
 	{
 		int32 SizeOfBufferInFrames = 0;
@@ -62,7 +69,7 @@ public:
 	{
 		TWeakObjectPtr<const USoundSubmix> Submix;
 		FAudioDevice* Device = nullptr;
-		TWeakObjectPtr<UAudioLinkSettingsAbstract> Settings;
+		TWeakObjectPtr<const UAudioLinkSettingsAbstract> Settings;
 	};
 
 	/**
@@ -88,12 +95,35 @@ public:
 	 * @return The newly created Link instance (if successful).
 	 */
 	virtual TUniquePtr<IAudioLink> CreateSourceAudioLink(const FAudioLinkSourceCreateArgs&) = 0;
+	
+	struct FAudioLinkSourcePushedCreateArgs
+	{
+		FName OwnerName;
+		int32 NumChannels = INDEX_NONE;
+		int32 SampleRate = INDEX_NONE;
+		int32 NumFramesPerBuffer = INDEX_NONE;
+		UAudioLinkSettingsAbstract::FSharedSettingsProxyPtr Settings;
+	};	
+	using FAudioLinkSourcePushedSharedPtr = TSharedPtr<IAudioLinkSourcePushed, ESPMode::ThreadSafe>;
+	virtual FAudioLinkSourcePushedSharedPtr CreateSourcePushedAudioLink(const FAudioLinkSourcePushedCreateArgs&) = 0;
+
+	/**
+	 * Create a AudioLinkSynchronizer callback
+	 * @return The newly created Link instance (if successful).
+	 */
+	using FAudioLinkSynchronizerSharedPtr = TSharedPtr<IAudioLinkSynchronizer, ESPMode::ThreadSafe>;
+	virtual FAudioLinkSynchronizerSharedPtr CreateSynchronizerAudioLink() = 0;
 		
 	/**	
 	 * Get the name of all AudioLink factories in the Modular Features registry.
 	 * @return "AudioLink Factory"
 	 */
-	static FName GetModularFeatureName();
+
+	inline static FName GetModularFeatureName()
+	{
+		static FName FeatureName = FName(TEXT("AudioLink Factory"));
+		return FeatureName;
+	}
 
 	/**
 	 * Gets all registered factory instances.
diff --git a/Engine/Source/Runtime/AudioLink/AudioLinkEngine/Public/IBufferedAudioOutput.h b/Engine/Source/Runtime/AudioLink/AudioLinkEngine/Public/IBufferedAudioOutput.h
--- a/Engine/Source/Runtime/AudioLink/AudioLinkEngine/Public/IBufferedAudioOutput.h
+++ b/Engine/Source/Runtime/AudioLink/AudioLinkEngine/Public/IBufferedAudioOutput.h
@@ -2,6 +2,21 @@
 
 #pragma once
 
+class IPushableAudioOutput
+{
+public:
+	struct FOnNewBufferParams
+	{
+		const float* AudioData = nullptr;
+		int32 Id = INDEX_NONE;
+		int32 NumSamples = INDEX_NONE;
+		int32 NumChannels = INDEX_NONE;
+		int32 SampleRate = INDEX_NONE;
+	};
+	virtual void PushNewBuffer(const FOnNewBufferParams&) = 0;	
+	virtual void LastBuffer(int32 InId) = 0;
+};
+
 /**
 	Abstract interface for communication with outputting audio objects
 	Examples concrete implementation of these are (source, submix etc).
@@ -75,7 +90,11 @@ public:
 	 * @param InNumSamplesToReserve Reserve this number of samples.
 	 */
 	virtual void Reserve(int32 InNumSamplesToReserve, int32 InNumSamplesOfSilence = 0) = 0;
+
+	virtual IPushableAudioOutput* GetPushableInterface() { return nullptr; }
+	virtual const IPushableAudioOutput* GetPushableInterface() const { return nullptr; }
 };
 
 using FWeakBufferedOutputPtr = TWeakPtr<IBufferedAudioOutput, ESPMode::ThreadSafe>;
-using FSharedBufferedOutputPtr = TSharedPtr<IBufferedAudioOutput, ESPMode::ThreadSafe>;
\ No newline at end of file
+using FSharedBufferedOutputPtr = TSharedPtr<IBufferedAudioOutput, ESPMode::ThreadSafe>;
+
diff --git a/Engine/Source/Runtime/AudioLink/AudioMixerPlatformAudioLink/AudioMixerPlatformAudioLink.Build.cs b/Engine/Source/Runtime/AudioLink/AudioMixerPlatformAudioLink/AudioMixerPlatformAudioLink.Build.cs
--- a/Engine/Source/Runtime/AudioLink/AudioMixerPlatformAudioLink/AudioMixerPlatformAudioLink.Build.cs
+++ b/Engine/Source/Runtime/AudioLink/AudioMixerPlatformAudioLink/AudioMixerPlatformAudioLink.Build.cs
@@ -22,7 +22,6 @@ public class AudioMixerPlatformAudioLink: ModuleRules
 					
 		if (Target.bCompileAgainstEngine) 
 		{
-			//PrivateDependencyModuleNames.Add("XMA2");
 			PrivateDependencyModuleNames.Add("BinkAudioDecoder");
 
 			AddEngineThirdPartyPrivateStaticDependencies(Target,
diff --git a/Engine/Source/Runtime/AudioLink/AudioMixerPlatformAudioLink/Private/AudioMixerPlatformAudioLink.h b/Engine/Source/Runtime/AudioLink/AudioMixerPlatformAudioLink/Private/AudioMixerPlatformAudioLink.h
--- a/Engine/Source/Runtime/AudioLink/AudioMixerPlatformAudioLink/Private/AudioMixerPlatformAudioLink.h
+++ b/Engine/Source/Runtime/AudioLink/AudioMixerPlatformAudioLink/Private/AudioMixerPlatformAudioLink.h
@@ -3,6 +3,7 @@
 #pragma once
 
 #include "AudioMixer.h"
+#include "IAudioLinkFactory.h"
 
 namespace Audio
 {
@@ -33,6 +34,7 @@ namespace Audio
 	private:
 		bool bSuspended = false;
 		bool bInitialized = false;
+		IAudioLinkFactory::FAudioLinkSynchronizerSharedPtr SynchronizeLink;
 	};
 }
 
diff --git a/Engine/Source/Runtime/AudioMixer/AudioMixer.Build.cs b/Engine/Source/Runtime/AudioMixer/AudioMixer.Build.cs
--- a/Engine/Source/Runtime/AudioMixer/AudioMixer.Build.cs
+++ b/Engine/Source/Runtime/AudioMixer/AudioMixer.Build.cs
@@ -31,6 +31,7 @@ namespace UnrealBuildTool.Rules
 					"AudioPlatformConfiguration",
 					"SoundFieldRendering",
 					"AudioExtensions",
+					"AudioLinkCore",
 					"AudioLinkEngine",
 					"HeadMountedDisplay"
 				}
diff --git a/Engine/Source/Runtime/AudioMixer/Private/AudioMixerDevice.cpp b/Engine/Source/Runtime/AudioMixer/Private/AudioMixerDevice.cpp
--- a/Engine/Source/Runtime/AudioMixer/Private/AudioMixerDevice.cpp
+++ b/Engine/Source/Runtime/AudioMixer/Private/AudioMixerDevice.cpp
@@ -171,6 +171,15 @@ namespace Audio
 		bAudioMixerModuleLoaded = true;
 
 		SourceManager = MakeUnique<FMixerSourceManager>(this);
+
+		// Register AudioLink Factory. 	
+		TArray<FName> Factories = IAudioLinkFactory::GetAllRegisteredFactoryNames();
+		if(Factories.Num() > 0)
+		{
+			// Allow only a single registered factory instance for now.
+			check(Factories.Num()==1);
+			AudioLinkFactory=IAudioLinkFactory::FindFactory(Factories[0]);
+		}
 	}
 
 	FMixerDevice::~FMixerDevice()
@@ -1953,6 +1962,11 @@ namespace Audio
 		return Sources.Num();
 	}
 
+	IAudioLinkFactory* FMixerDevice::GetAudioLinkFactory() const
+	{
+		return AudioLinkFactory;
+	}
+
 	int32 FMixerDevice::GetNumActiveSources() const
 	{
 		return SourceManager->GetNumActiveSources();
diff --git a/Engine/Source/Runtime/AudioMixer/Private/AudioMixerSource.cpp b/Engine/Source/Runtime/AudioMixer/Private/AudioMixerSource.cpp
--- a/Engine/Source/Runtime/AudioMixer/Private/AudioMixerSource.cpp
+++ b/Engine/Source/Runtime/AudioMixer/Private/AudioMixerSource.cpp
@@ -596,6 +596,28 @@ namespace Audio
 			InitParams.SourceBufferListener = WaveInstance->SourceBufferListener;
 			InitParams.bShouldSourceBufferListenerZeroBuffer = WaveInstance->bShouldSourceBufferListenerZeroBuffer;
 
+			if( WaveInstance->bShouldUseAudioLink)
+			{
+				if (IAudioLinkFactory* LinkFactory = MixerDevice->GetAudioLinkFactory())
+				{				
+					IAudioLinkFactory::FAudioLinkSourcePushedCreateArgs CreateArgs;					
+					if (WaveInstance->AudioLinkSettingsOverride)
+					{
+						CreateArgs.Settings = WaveInstance->AudioLinkSettingsOverride->GetProxy();
+					}
+					else
+					{
+						CreateArgs.Settings = GetDefault<UAudioLinkSettingsAbstract>(LinkFactory->GetSettingsClass())->GetProxy();
+					}
+					
+					CreateArgs.OwnerName = *WaveInstance->GetName();			// <-- FIXME: String FName conversion.
+					CreateArgs.NumChannels = SoundBuffer->NumChannels;
+					CreateArgs.NumFramesPerBuffer = MixerDevice->GetBufferLength();
+					CreateArgs.SampleRate = MixerDevice->GetSampleRate();
+					AudioLink = LinkFactory->CreateSourcePushedAudioLink(CreateArgs);
+					InitParams.AudioLink = AudioLink;
+				}
+			}
 
 			// Source manager needs to know if this is a vorbis source for rebuilding speaker maps
 			InitParams.bIsVorbis = bIsVorbis;
@@ -968,6 +990,14 @@ namespace Audio
 			AudioDevice->SourceDataOverridePluginInterface->GetSourceDataOverrides(SourceId, ListenerTransform, WaveInstance);
 		}
 
+		// AudioLink, push state if we're enabled and 3d.
+		if (bIs3D && AudioLink.IsValid())
+		{
+			IAudioLinkSourcePushed::FOnUpdateWorldStateParams Params;
+			Params.WorldTransform = WaveInstance->ActiveSound->Transform;
+			AudioLink->OnUpdateWorldState(Params);
+		}
+
 		UpdatePitch();
 
 		UpdateVolume();
@@ -1407,6 +1437,11 @@ namespace Audio
 		check(!bIsStopping);
 		check(!Playing);
 
+		if (AudioLink.IsValid())
+		{
+			AudioLink.Reset();
+		}
+
 		// Make a new pending release data ptr to pass off release data
 		if (MixerSourceVoice)
 		{
diff --git a/Engine/Source/Runtime/AudioMixer/Private/AudioMixerSource.h b/Engine/Source/Runtime/AudioMixer/Private/AudioMixerSource.h
--- a/Engine/Source/Runtime/AudioMixer/Private/AudioMixerSource.h
+++ b/Engine/Source/Runtime/AudioMixer/Private/AudioMixerSource.h
@@ -118,6 +118,7 @@ namespace Audio
 		FMixerBuffer* MixerBuffer;
 		TSharedPtr<FMixerSourceBuffer, ESPMode::ThreadSafe> MixerSourceBuffer;
 		FMixerSourceVoice* MixerSourceVoice;
+		IAudioLinkFactory::FAudioLinkSourcePushedSharedPtr AudioLink;
 
 		uint32 bPreviousBusEnablement;
 		uint32 bPreviousBaseSubmixEnablement;
diff --git a/Engine/Source/Runtime/AudioMixer/Private/AudioMixerSourceManager.cpp b/Engine/Source/Runtime/AudioMixer/Private/AudioMixerSourceManager.cpp
--- a/Engine/Source/Runtime/AudioMixer/Private/AudioMixerSourceManager.cpp
+++ b/Engine/Source/Runtime/AudioMixer/Private/AudioMixerSourceManager.cpp
@@ -606,6 +606,12 @@ namespace Audio
 		SourceInfo.LowPassFreq = MAX_FILTER_FREQUENCY;
 		SourceInfo.HighPassFreq = MIN_FILTER_FREQUENCY;
 
+		if (SourceInfo.AudioLink)
+		{
+			SourceInfo.AudioLink->OnSourceReleased(SourceId);
+			SourceInfo.AudioLink.Reset();
+		}
+
 		if (SourceInfo.SourceBufferListener)
 		{
 			SourceInfo.SourceBufferListener->OnSourceReleased(SourceId);
@@ -909,6 +915,11 @@ namespace Audio
 				SourceInfo.bUseReverbPlugin = true;
 			}
 
+			if (InitParams.AudioLink.IsValid())
+			{
+				SourceInfo.AudioLink = InitParams.AudioLink;
+			}
+
 			// Optional Source Buffer listener.
 			SourceInfo.SourceBufferListener = InitParams.SourceBufferListener;
 			SourceInfo.bShouldSourceBufferListenerZeroBuffer = InitParams.bShouldSourceBufferListenerZeroBuffer;
@@ -2204,6 +2215,14 @@ namespace Audio
 			return;
 		}
 
+		if (SourceInfo.AudioLink.IsValid())
+		{
+			IAudioLinkSourcePushed::FOnNewBufferParams Params;
+			Params.SourceId = SourceId;
+			Params.Buffer = SourceInfo.PreDistanceAttenuationBuffer;
+			SourceInfo.AudioLink->OnNewBuffer(Params);
+		}
+
 		// If we have Source Buffer Listener
 		if (SourceInfo.SourceBufferListener.IsValid())
 		{
diff --git a/Engine/Source/Runtime/AudioMixer/Private/AudioMixerSourceManager.h b/Engine/Source/Runtime/AudioMixer/Private/AudioMixerSourceManager.h
--- a/Engine/Source/Runtime/AudioMixer/Private/AudioMixerSourceManager.h
+++ b/Engine/Source/Runtime/AudioMixer/Private/AudioMixerSourceManager.h
@@ -145,6 +145,8 @@ namespace Audio
 
 		FSharedISourceBufferListenerPtr SourceBufferListener;		
 
+		IAudioLinkFactory::FAudioLinkSourcePushedSharedPtr AudioLink;
+
 		FName AudioComponentUserID;
 		uint64 AudioComponentID = 0;
 		bool bIs3D = false;
@@ -382,6 +384,8 @@ namespace Audio
 		// Cached ptr to an optional source data override plugin
 		TAudioSourceDataOverridePtr SourceDataOverridePlugin;
 
+		IAudioLinkFactory* AudioLinkFactory = nullptr;
+
 		// Array of pointers to game thread audio source objects
 		TArray<FMixerSourceVoice*> MixerSources;
 
@@ -499,6 +503,9 @@ namespace Audio
 			// Optional Source buffer listener.
 			FSharedISourceBufferListenerPtr SourceBufferListener;
 
+			// Optional AudioLink.
+			IAudioLinkFactory::FAudioLinkSourcePushedSharedPtr AudioLink;
+
 			// State management
 			uint8 bIs3D:1;
 			uint8 bIsCenterChannelOnly:1;
diff --git a/Engine/Source/Runtime/AudioMixer/Private/AudioMixerSubmix.cpp b/Engine/Source/Runtime/AudioMixer/Private/AudioMixerSubmix.cpp
--- a/Engine/Source/Runtime/AudioMixer/Private/AudioMixerSubmix.cpp
+++ b/Engine/Source/Runtime/AudioMixer/Private/AudioMixerSubmix.cpp
@@ -284,15 +284,20 @@ namespace Audio
 					UpdateModulationSettings(VolMod, WetMod, DryMod);
 				});
 			}
-			if (SoundSubmix->AudioLinkSettings)
+			
+			// AudioLink send enabled? 
+			if (SoundSubmix->bSendToAudioLink)
 			{
-				if (IAudioLinkFactory* LinkFactory = IAudioLinkFactory::FindFactory(SoundSubmix->AudioLinkSettings->GetFactoryName()))
-				{
-					AudioLinkInstance = LinkFactory->CreateSubmixAudioLink({ SoundSubmix, MixerDevice, SoundSubmix->AudioLinkSettings });
-				}
-				else
+				// If AudioLink is active, create a link.
+				if (IAudioLinkFactory* LinkFactory = MixerDevice->GetAudioLinkFactory())
 				{
-					UE_LOG(LogAudioLink, Warning, TEXT("Failed to Find AudioLink Factory '%s'"), *SoundSubmix->AudioLinkSettings->GetFactoryName().GetPlainNameString());
+					check(LinkFactory->GetSettingsClass());
+					check(LinkFactory->GetSettingsClass()->GetDefaultObject());
+
+					const UAudioLinkSettingsAbstract* Settings = SoundSubmix->AudioLinkSettings.IsNull() ?
+						GetDefault<UAudioLinkSettingsAbstract>(LinkFactory->GetSettingsClass()) : SoundSubmix->AudioLinkSettings.Get();
+
+					AudioLinkInstance = LinkFactory->CreateSubmixAudioLink({ SoundSubmix, MixerDevice, Settings });
 				}
 			}
 			
diff --git a/Engine/Source/Runtime/AudioMixer/Public/AudioMixerDevice.h b/Engine/Source/Runtime/AudioMixer/Public/AudioMixerDevice.h
--- a/Engine/Source/Runtime/AudioMixer/Public/AudioMixerDevice.h
+++ b/Engine/Source/Runtime/AudioMixer/Public/AudioMixerDevice.h
@@ -17,6 +17,7 @@
 class FOnSubmixEnvelopeBP;
 class IAudioMixerPlatformInterface;
 class USoundModulatorBase;
+class IAudioLinkFactory;
 
 namespace Audio
 {
@@ -218,6 +219,9 @@ namespace Audio
 		void ReleaseMixerSourceVoice(FMixerSourceVoice* InSourceVoice);
 		int32 GetNumSources() const;
 
+		// AudioLink
+		IAudioLinkFactory* GetAudioLinkFactory() const;
+
 		const FAudioPlatformDeviceInfo& GetPlatformDeviceInfo() const { return PlatformInfo; };
 
 		FORCEINLINE int32 GetNumDeviceChannels() const { return PlatformInfo.NumChannels; }
@@ -435,6 +439,8 @@ namespace Audio
 
 		/** MPSC command queue to send commands to the game thread */
 		TMpscQueue<TFunction<void()>> GameThreadCommandQueue;
+
+		IAudioLinkFactory* AudioLinkFactory = nullptr;
 		
 		/** Whether or not we generate output audio to test multi-platform mixer. */
 		bool bDebugOutputEnabled;
diff --git a/Engine/Source/Runtime/Engine/Classes/Sound/SoundSubmix.h b/Engine/Source/Runtime/Engine/Classes/Sound/SoundSubmix.h
--- a/Engine/Source/Runtime/Engine/Classes/Sound/SoundSubmix.h
+++ b/Engine/Source/Runtime/Engine/Classes/Sound/SoundSubmix.h
@@ -307,6 +307,10 @@ public:
 	UPROPERTY(EditAnywhere, BlueprintReadWrite, Category = SubmixLevel, meta = (DisplayName = "Dry Level (dB)", AudioParam = "Volume", ClampMin = "-96.0", ClampMax = "0.0", UIMin = "-96.0", UIMax = "0.0"))
 	FSoundModulationDestinationSettings DryLevelModulation;
 	
+	/** Whether to send this Submix to AudioLink (when AudioLink is Enabled)*/
+	UPROPERTY(EditAnywhere, BlueprintReadWrite, Category = AudioLink)
+	uint8 bSendToAudioLink : 1;
+
 	/** Optional Audio Link Settings Object */
 	UPROPERTY(EditAnywhere, BlueprintReadWrite, Category = AudioLink)
 	TObjectPtr<UAudioLinkSettingsAbstract> AudioLinkSettings;
diff --git a/Engine/Source/Runtime/Engine/Private/AudioDeviceManager.cpp b/Engine/Source/Runtime/Engine/Private/AudioDeviceManager.cpp
--- a/Engine/Source/Runtime/Engine/Private/AudioDeviceManager.cpp
+++ b/Engine/Source/Runtime/Engine/Private/AudioDeviceManager.cpp
@@ -33,6 +33,8 @@
 #include "Settings/LevelEditorMiscSettings.h"
 #endif
 
+#include "IAudioLinkFactory.h"
+
 static int32 GCVarEnableAudioThreadWait = 1;
 TAutoConsoleVariable<int32> CVarEnableAudioThreadWait(
 	TEXT("AudioThread.EnableAudioThreadWait"),
@@ -658,6 +660,13 @@ bool FAudioDeviceManager::LoadDefaultAudioDeviceModule()
 		// Get the audio mixer and non-audio mixer device module names
 		GConfig->GetString(TEXT("Audio"), TEXT("AudioDeviceModuleName"), AudioDeviceModuleName, GEngineIni);
 		GConfig->GetString(TEXT("Audio"), TEXT("AudioMixerModuleName"), AudioMixerModuleName, GEngineIni);
+
+		// If AudioLink is Enabled, use the AudioLink Platform.
+		static const FString AudioLinkMixerModule = TEXT("AudioMixerPlatformAudioLink");
+		if (IModularFeatures::Get().IsModularFeatureAvailable(TEXT("AudioLink Factory")))
+		{
+			AudioMixerModuleName = AudioLinkMixerModule;
+		}
 	}
 	else if (bForceNoAudioMixer)
 	{
diff --git a/Engine/Source/Runtime/Engine/Public/Audio.h b/Engine/Source/Runtime/Engine/Public/Audio.h
--- a/Engine/Source/Runtime/Engine/Public/Audio.h
+++ b/Engine/Source/Runtime/Engine/Public/Audio.h
@@ -63,6 +63,7 @@ class USoundWave;
 class USoundClass;
 class USoundSubmix;
 class USoundSourceBus;
+class UAudioLinkSettingsAbstract;
 struct FActiveSound;
 struct FWaveInstance;
 struct FSoundSourceBusSendInfo;
@@ -213,6 +214,10 @@ struct ENGINE_API FWaveInstance
 	FSharedISourceBufferListenerPtr SourceBufferListener;
 	bool bShouldSourceBufferListenerZeroBuffer = false;
 
+	/** AudioLink Opt in */
+	bool bShouldUseAudioLink = true;
+	UAudioLinkSettingsAbstract* AudioLinkSettingsOverride = nullptr;	
+
 private:
 
 	/** Current volume */
