AudioLink Fixes for thread safety

diff --git a/Engine/Source/Runtime/AudioLink/AudioLinkEngine/Private/IAudioLinkFactory.cpp b/Engine/Source/Runtime/AudioLink/AudioLinkEngine/Private/IAudioLinkFactory.cpp
--- a/Engine/Source/Runtime/AudioLink/AudioLinkEngine/Private/IAudioLinkFactory.cpp
+++ b/Engine/Source/Runtime/AudioLink/AudioLinkEngine/Private/IAudioLinkFactory.cpp
@@ -7,9 +7,22 @@
 #include "BufferedSubmixListener.h" 
 #include "BufferedSourceListener.h" 
 
+namespace AudioLinkFactory_Private
+{
+	bool IsEnabled()
+	{
+		static bool bAudioLinkEnabled = true;
+		static FAutoConsoleVariableRef CVarAudioLinkEnabled(TEXT("au.audiolink.enabled"), bAudioLinkEnabled, TEXT("Enable AudioLink"), ECVF_Default);
+		return bAudioLinkEnabled;
+	}
+}
+
 IAudioLinkFactory::IAudioLinkFactory()
 {
-	IModularFeatures::Get().RegisterModularFeature(GetModularFeatureName(), this);
+	if (AudioLinkFactory_Private::IsEnabled())
+	{
+		IModularFeatures::Get().RegisterModularFeature(GetModularFeatureName(), this);
+	}
 }
 
 IAudioLinkFactory::~IAudioLinkFactory()
diff --git a/Engine/Source/Runtime/AudioLink/AudioLinkEngine/Public/IAudioLink.h b/Engine/Source/Runtime/AudioLink/AudioLinkEngine/Public/IAudioLink.h
--- a/Engine/Source/Runtime/AudioLink/AudioLinkEngine/Public/IAudioLink.h
+++ b/Engine/Source/Runtime/AudioLink/AudioLinkEngine/Public/IAudioLink.h
@@ -52,11 +52,43 @@ class IAudioLinkSynchronizer
 protected:
 	IAudioLinkSynchronizer() = default;
 public:
-	virtual ~IAudioLinkSynchronizer() = default;
+	virtual ~IAudioLinkSynchronizer() = default;	
+	
+	DECLARE_MULTICAST_DELEGATE(FOnSuspend);
+	virtual FDelegateHandle RegisterSuspendDelegate(const FOnSuspend::FDelegate&) = 0;
+	virtual bool RemoveSuspendDelegate(const FDelegateHandle&) = 0;
 
-	DECLARE_DELEGATE(FOnBeginRender);
-	virtual FOnBeginRender& GetOnBeginRender() = 0;
+	DECLARE_MULTICAST_DELEGATE(FOnResume);
+	virtual FDelegateHandle RegisterResumeDelegate(const FOnResume::FDelegate&) = 0;
+	virtual bool RemoveResumeDelegate(const FDelegateHandle&) = 0;
 
-	DECLARE_DELEGATE(FOnEndRender);	
-	virtual FOnEndRender& GetOnEndRender() = 0;
+	struct FOnOpenStreamParams
+	{
+		FString Name;
+		int32 NumFrames = INDEX_NONE; 
+		int32 NumChannels = INDEX_NONE;
+		int32 SampleRate = INDEX_NONE;
+	};
+
+	DECLARE_MULTICAST_DELEGATE_OneParam(FOnOpenStream, FOnOpenStreamParams);
+	virtual FDelegateHandle RegisterOpenStreamDelegate(const FOnOpenStream::FDelegate&) = 0;
+	virtual bool RemoveOpenStreamDelegate(const FDelegateHandle&) = 0;
+
+	DECLARE_MULTICAST_DELEGATE(FOnCloseStream);
+	virtual FDelegateHandle RegisterCloseStreamDelegate(const FOnCloseStream::FDelegate&) = 0;
+	virtual bool RemoveCloseStreamDelegate(const FDelegateHandle&) = 0;
+
+	struct FOnRenderParams 
+	{
+		uint64 FrameID = 0;
+		int32 NumFrames = INDEX_NONE;
+	};
+	
+	DECLARE_MULTICAST_DELEGATE_OneParam(FOnBeginRender, FOnRenderParams);
+	virtual FDelegateHandle RegisterBeginRenderDelegate(const FOnBeginRender::FDelegate&) = 0;
+	virtual bool RemoveBeginRenderDelgate(const FDelegateHandle&) = 0;
+
+	DECLARE_MULTICAST_DELEGATE(FOnEndRender);	
+	virtual FDelegateHandle RegisterEndRenderDelegate(const FOnEndRender::FDelegate&) = 0;
+	virtual bool RemoveEndRenderDelgate(const FDelegateHandle&) = 0;
 };
\ No newline at end of file
