AudioLink Engine fixes

diff --git a/Engine/Source/Editor/DetailCustomizations/Private/AttenuationSettingsCustomizations.cpp b/Engine/Source/Editor/DetailCustomizations/Private/AttenuationSettingsCustomizations.cpp
--- a/Engine/Source/Editor/DetailCustomizations/Private/AttenuationSettingsCustomizations.cpp
+++ b/Engine/Source/Editor/DetailCustomizations/Private/AttenuationSettingsCustomizations.cpp
@@ -440,6 +440,7 @@ void FSoundAttenuationSettingsCustomization::CustomizeChildren(TSharedRef<IPrope
 	bIsPriorityAttenuationEnabledHandle = PropertyHandles.FindChecked(GET_MEMBER_NAME_CHECKED(FSoundAttenuationSettings, bEnablePriorityAttenuation)).ToSharedRef();
 	bIsSubmixSendAttenuationEnabledHandle = PropertyHandles.FindChecked(GET_MEMBER_NAME_CHECKED(FSoundAttenuationSettings, bEnableSubmixSends)).ToSharedRef();
 	bIsSourceDataOverrideEnabledHandle = PropertyHandles.FindChecked(GET_MEMBER_NAME_CHECKED(FSoundAttenuationSettings, bEnableSourceDataOverride)).ToSharedRef();
+	bIsSendToAudioLinkEnabledHandle = PropertyHandles.FindChecked(GET_MEMBER_NAME_CHECKED(FSoundAttenuationSettings, bEnableSendToAudioLink)).ToSharedRef();
 	ReverbSendMethodHandle = PropertyHandles.FindChecked(GET_MEMBER_NAME_CHECKED(FSoundAttenuationSettings, ReverbSendMethod)).ToSharedRef();
 	PriorityAttenuationMethodHandle = PropertyHandles.FindChecked(GET_MEMBER_NAME_CHECKED(FSoundAttenuationSettings, PriorityAttenuationMethod)).ToSharedRef();
 	AbsorptionMethodHandle = PropertyHandles.FindChecked(GET_MEMBER_NAME_CHECKED(FSoundAttenuationSettings, AbsorptionMethod)).ToSharedRef();
@@ -681,6 +682,14 @@ void FSoundAttenuationSettingsCustomization::CustomizeChildren(TSharedRef<IPrope
 		.Visibility(IsAttenuationOverriddenVisibleAttribute())
 		.EditCondition(IsAttenuationOverriddenAttribute(), nullptr);
 
+	LayoutBuilder.AddPropertyToCategory(PropertyHandles.FindChecked(GET_MEMBER_NAME_CHECKED(FSoundAttenuationSettings, bEnableSendToAudioLink)))
+		.Visibility(IsAttenuationOverriddenVisibleAttribute())
+		.EditCondition(IsAttenuationOverriddenAttribute(), nullptr);
+
+	LayoutBuilder.AddPropertyToCategory(PropertyHandles.FindChecked(GET_MEMBER_NAME_CHECKED(FSoundAttenuationSettings, AudioLinkSettingsOverride)))
+		.Visibility(IsAttenuationOverriddenVisibleAttribute())
+		.EditCondition(IsAttenuationOverriddenAttribute(), nullptr);
+
 	if (bIsAudioMixerEnabled)
 	{
 		TSharedPtr<IPropertyHandle> PluginProperty = PropertyHandles.FindChecked(GET_MEMBER_NAME_CHECKED(FSoundAttenuationSettings, PluginSettings));
@@ -705,9 +714,10 @@ void FSoundAttenuationSettingsCustomization::CustomizeChildren(TSharedRef<IPrope
 	LayoutBuilder.EditCategory("AttenuationAirAbsorption", LOCTEXT("AttenuationAirAbsorption", "Attenuation (Air Absorption)"));
 	LayoutBuilder.EditCategory("AttenuationPluginSettings", LOCTEXT("AttenuationPluginSettings", "Attenuation (Plugin Settings)"));
 	LayoutBuilder.EditCategory("AttenuationSourceDataOverride", LOCTEXT("AttenuationSourceDataOverride", "Attenuation (Source Data Override)"));
+	LayoutBuilder.EditCategory("AttenuationAudioLink", LOCTEXT("AttenuationAudioLink", "Attenuation (AudioLink)"));
 	LayoutBuilder.SortCategories(AttenuationSettingsUtils::SortCategories);
 
-	if (PropertyHandles.Num() != 66)
+	if (PropertyHandles.Num() != 68)
 	{
 		ensureMsgf(false, TEXT("Unexpected property handle(s) customizing FSoundAttenuationSettings. %d handles found"), PropertyHandles.Num());
 	}
diff --git a/Engine/Source/Editor/DetailCustomizations/Private/AttenuationSettingsCustomizations.h b/Engine/Source/Editor/DetailCustomizations/Private/AttenuationSettingsCustomizations.h
--- a/Engine/Source/Editor/DetailCustomizations/Private/AttenuationSettingsCustomizations.h
+++ b/Engine/Source/Editor/DetailCustomizations/Private/AttenuationSettingsCustomizations.h
@@ -60,6 +60,8 @@ protected:
 	TSharedPtr<IPropertyHandle> bIsPriorityAttenuationEnabledHandle;
 	TSharedPtr<IPropertyHandle> bIsSubmixSendAttenuationEnabledHandle;
 	TSharedPtr<IPropertyHandle> bIsSourceDataOverrideEnabledHandle;
+	TSharedPtr<IPropertyHandle> bIsSendToAudioLinkEnabledHandle;
+	TSharedPtr<IPropertyHandle> bEnableSourceDataOverrideHandle;
 	TSharedPtr<IPropertyHandle> ReverbSendMethodHandle;
 	TSharedPtr<IPropertyHandle> PriorityAttenuationMethodHandle;
 	TSharedPtr<IPropertyHandle> AbsorptionMethodHandle;
diff --git a/Engine/Source/Runtime/AudioLink/AudioLinkEngine/Private/BufferedListenerBase.cpp b/Engine/Source/Runtime/AudioLink/AudioLinkEngine/Private/BufferedListenerBase.cpp
--- a/Engine/Source/Runtime/AudioLink/AudioLinkEngine/Private/BufferedListenerBase.cpp
+++ b/Engine/Source/Runtime/AudioLink/AudioLinkEngine/Private/BufferedListenerBase.cpp
@@ -4,16 +4,21 @@
 #include "AudioLinkLog.h"
 
 FBufferedListenerBase::FBufferedListenerBase(int32 InDefaultCircularBufferSize)
+	: CircularBuffer(InDefaultCircularBufferSize)
 {
-	LocklessCircularBuffer.SetCapacity(InDefaultCircularBufferSize);
+	CircularBuffer.SetCapacity(InDefaultCircularBufferSize);
 }
 
 /** CONSUMER Thread. This will be called by the consuming data of the buffers. */
 bool FBufferedListenerBase::PopBuffer(float* InBuffer, int32 InBufferSizeInSamples, int32& OutSamplesWritten)
 {
-	OutSamplesWritten = LocklessCircularBuffer.Pop(InBuffer, InBufferSizeInSamples);
-	NumSilentSamplesRemaining = FMath::Max(0,NumSilentSamplesRemaining-OutSamplesWritten);
-	return bStarted;
+	FScopeLock Lock(&CircularBuffer.GetCriticialSection());
+	OutSamplesWritten = CircularBuffer.Pop(InBuffer, InBufferSizeInSamples);
+	int32 NumSamplesRemaining = CircularBuffer.Num();
+	
+	// Return true if any data remaining, which signals the consumer to shutdown if false.
+	// However if we're still running, we could be starving, so only signal stopped if we've exhausted.
+	return (bStopping && NumSamplesRemaining > 0) || bStarted;
 }
 
 /** CONSUMER Thread. This will be called by the consuming data of the buffers. */
@@ -72,16 +77,27 @@ void FBufferedListenerBase::OnBufferReceived(const FBufferFormat& InFormat, TArr
 		OnFormatKnown.ExecuteIfBound(InFormat);
 	}
 
-	// Push the data into the circular buffer.
-	int32 SamplesPushed = LocklessCircularBuffer.Push(InBuffer.GetData(), InBuffer.Num());
+	int32 SamplesPushed = 0;
+	{
+		// Lock outside for the sake of atomic logging
+		FScopeLock Lock(&CircularBuffer.GetCriticialSection());
+		
+		// Push the data into the circular buffer.
+		SamplesPushed = CircularBuffer.Push(InBuffer.GetData(), InBuffer.Num());
+
+		UE_LOG(LogAudioLink, VeryVerbose, 
+			TEXT("FBufferedListenerBase::OnBufferReceived()(post-push), SamplesPushed=%d, InBuffer.Num()=%d, LocklessCircularBuffer.Num()=%d, LocklessCircularBuffer.Remainder()=%d, LocklessCircularBuffer.Capacity()=%d, This=0x%p"),
+				SamplesPushed, InBuffer.Num(), CircularBuffer.Num(), CircularBuffer.Remainder(), CircularBuffer.GetCapacity(), this );
+	}
 
 	// Warn of not enough space in circular buffer, unless we're overwriting silence.
-	if (SamplesPushed < InBuffer.Num() && NumSilentSamplesRemaining == 0)
+	if (SamplesPushed < InBuffer.Num())
 	{
 		// Prevent log spam by limiting to 1:100 logs.
 		static const int32 NumLogMessagesToSkip = 100;
 		static int32 LogPacifier = 0;
-		UE_CLOG(LogPacifier++ % NumLogMessagesToSkip == 0, LogAudioLink, Log, TEXT("Overflow by '%d' Samples in Buffer Listener"), InBuffer.Num() - SamplesPushed);
+		UE_CLOG(LogPacifier++ % NumLogMessagesToSkip == 0, LogAudioLink, 
+			Verbose, TEXT("FBufferedListenerBase: Overflow by '%d' Samples in Buffer Listener"), InBuffer.Num() - SamplesPushed);
 	}
 }
 
@@ -118,20 +134,26 @@ bool FBufferedListenerBase::TryUnsetStartedFlag()
 	return bStarted.compare_exchange_strong(bExpected, bDesired);
 }
 
+bool FBufferedListenerBase::TrySetStoppingFlag()
+{
+	bool bExpected = false;
+	bool bDesired = true;
+	return bStopping.compare_exchange_strong(bExpected, bDesired);
+}
+
 void FBufferedListenerBase::PushSilence(int32 InNumSamplesOfSilence)
 {
-	LocklessCircularBuffer.PushZeros(InNumSamplesOfSilence);
-	NumSilentSamplesRemaining += InNumSamplesOfSilence;
+	CircularBuffer.PushZeros(InNumSamplesOfSilence);
 }
 
 void FBufferedListenerBase::Reserve(int32 InNumSamplesToReserve, int32 InNumSamplesOfSilence)
 {
 	// This Zeros the buffer also so should only be done at the start.
-	LocklessCircularBuffer.SetCapacity(InNumSamplesToReserve);	
-
+	CircularBuffer.SetCapacity(InNumSamplesToReserve);	
+	
 	// Optionally add silence into the buffer, this will cause latency.
 	if (InNumSamplesOfSilence > 0)
 	{
 		PushSilence(InNumSamplesOfSilence);
 	}
-}
\ No newline at end of file
+}
diff --git a/Engine/Source/Runtime/AudioLink/AudioLinkEngine/Private/BufferedListenerBase.h b/Engine/Source/Runtime/AudioLink/AudioLinkEngine/Private/BufferedListenerBase.h
--- a/Engine/Source/Runtime/AudioLink/AudioLinkEngine/Private/BufferedListenerBase.h
+++ b/Engine/Source/Runtime/AudioLink/AudioLinkEngine/Private/BufferedListenerBase.h
@@ -7,7 +7,108 @@
 
 #include "IBufferedAudioOutput.h"
 #include "Misc/Optional.h"
-#include "Containers/ArrayView.h"
+#include "Math/UnrealMathUtility.h"
+
+// Simple locking circular buffer
+// Uses 32bit wrap around logic explained here: 
+// https://www.snellman.net/blog/archive/2016-12-13-ring-buffers/
+class FLockingCircularSampleBuffer
+{	
+public:
+	FLockingCircularSampleBuffer(int32 InInitialCapacity)
+	{
+		SetCapacity(InInitialCapacity);
+	}
+
+	// Allow the caller to lock the CS ahead of calling (if necessary)
+	FCriticalSection& GetCriticialSection() { return CS; }
+	const FCriticalSection& GetCriticialSection() const { return CS; }
+
+	// This doesn't preserve contents.
+	void SetCapacity(int32 InCapacity)
+	{
+		check(InCapacity > 0);
+
+		if(!FMath::IsPowerOfTwo(InCapacity))
+		{
+			InCapacity = FMath::RoundUpToPowerOfTwo(InCapacity);
+		}
+
+		FScopeLock Lock(&CS);
+		Buffer.SetNumZeroed(InCapacity, true /* Allow shrinking */);
+		Mask = InCapacity-1;
+		Read = 0; // Empty.
+		Write = 0;
+
+		check(Mask != 0);
+	}
+	int32 GetCapacity() const
+	{
+		FScopeLock Lock(&CS);
+		return Buffer.Num();
+	}
+	int32 Num() const
+	{
+		FScopeLock Lock(&CS);
+		return Write-Read;
+	}
+	
+	// Get number of samples that can be pushed onto the buffer before it is full.
+	int32 Remainder() const
+	{
+		FScopeLock Lock(&CS);
+		return Buffer.Num() - Num();
+	}
+
+	int32 Push(const float* InBuffer, int32 InSize)
+	{
+		FScopeLock Lock(&CS);
+
+		int32 CanPush = FMath::Min(Remainder(), InSize);
+		for (int32 i = 0; i < CanPush ; ++i)
+		{
+			Enqueue(InBuffer[i]);
+		}
+		return CanPush;		
+	}
+	int32 Pop(float* OutBuffer, int32 InNumSamples)
+	{
+		FScopeLock Lock(&CS);
+		int32 CanPop = FMath::Min(Num(),InNumSamples);
+		for(int32 i=0; i < CanPop; ++i)
+		{
+			OutBuffer[i] = Dequeue();
+		}
+		return CanPop;
+	}
+
+	int32 PushZeros(int32 InNumSamplesOfSilence)
+	{
+		FScopeLock Lock(&CS);
+		int32 CanPush = FMath::Min(GetCapacity(), InNumSamplesOfSilence);
+		for(int32 i = 0; i < CanPush; ++i)
+		{
+			Enqueue(0.f);
+		}
+		return CanPush;
+	}
+private:
+	uint32 Read = 0;		// These grow indefinitely until wrap at 2^32,
+	uint32 Write = 0;		// this allows us to use full capacity as write >= read.
+	uint32 Mask = 0;
+	TArray<float> Buffer;
+	mutable FCriticalSection CS;
+
+	// NOTE: Not forceinline as compiler does better job without.
+	void Enqueue(const float InFloat)
+	{
+		Buffer[Write++ & Mask] = InFloat;
+	}
+	float Dequeue()
+	{
+		return Buffer[Read++ & Mask];
+	}
+};
 
 /** Common base class of Buffered Listener objects.
 */
@@ -43,11 +144,13 @@ protected:
 	//* Attempt to set started to false. This can fail if we're already stopped. */	
 	bool TryUnsetStartedFlag();
 
+	bool TrySetStoppingFlag();
+
 private:	
 	void PushSilence(int32 InNumSamplesOfSilence);
 
-	/** Lock-less buffer to hold the data for the single source we're listening to, interleaved. */
-	Audio::TCircularAudioBuffer<float> LocklessCircularBuffer;
+	/** Buffer to hold the data for the single source we're listening to, interleaved. */
+	FLockingCircularSampleBuffer CircularBuffer;
 
 	/** Read/Write slim lock protects format known optional */
 	mutable FRWLock FormatKnownRwLock;
@@ -58,9 +161,9 @@ private:
 	/** Delegate that fires when the format it known. Normally on the first buffer received. */
 	FOnFormatKnown OnFormatKnown;
 
-	/** Count for how much (initial) silence is still remaining in the Circular buffer, for logging purposes */
-	std::atomic<int32> NumSilentSamplesRemaining;
-
 	/** Atomic flag we've been started */
 	std::atomic<bool> bStarted;
+
+	/** Atomic flag we've been told to stop */
+	std::atomic<bool> bStopping;
 };
diff --git a/Engine/Source/Runtime/AudioLink/AudioLinkEngine/Private/BufferedSourceListener.cpp b/Engine/Source/Runtime/AudioLink/AudioLinkEngine/Private/BufferedSourceListener.cpp
--- a/Engine/Source/Runtime/AudioLink/AudioLinkEngine/Private/BufferedSourceListener.cpp
+++ b/Engine/Source/Runtime/AudioLink/AudioLinkEngine/Private/BufferedSourceListener.cpp
@@ -20,7 +20,8 @@ bool FBufferedSourceListener::Start(FAudioDevice*)
 
 void FBufferedSourceListener::Stop(FAudioDevice*)
 {
-	ensure(TryUnsetStartedFlag());
+	TrySetStoppingFlag();
+	TryUnsetStartedFlag();
 }
 
 /** AUDIO MIXER THREAD. When a source is finished and returned to the pool, this call will be called. */
@@ -36,8 +37,9 @@ void FBufferedSourceListener::OnSourceReleased(const int32 InSourceId)
 		// Reset our source id.
 		CurrentSourceId = INDEX_NONE;
 
-		// Reset the format, now our source we were listening to has stopped.
-		ResetFormat();
+		// Mark us stopping.
+		Stop(nullptr);
+
 	}
 }
 
diff --git a/Engine/Source/Runtime/AudioLink/AudioLinkEngine/Public/IAudioLink.h b/Engine/Source/Runtime/AudioLink/AudioLinkEngine/Public/IAudioLink.h
--- a/Engine/Source/Runtime/AudioLink/AudioLinkEngine/Public/IAudioLink.h
+++ b/Engine/Source/Runtime/AudioLink/AudioLinkEngine/Public/IAudioLink.h
@@ -41,7 +41,11 @@ public:
 	};
 	virtual void OnNewBuffer(const FOnNewBufferParams&) = 0;
 
-	virtual void OnSourceReleased(const int32 SourceId) = 0;	
+	virtual void OnSourceDone(const int32 SourceId) = 0;
+
+	virtual void OnSourceReleased(const int32 SourceId) = 0;
+
+
 };
 
 /** IAudioLinkSynchronizer
@@ -70,7 +74,7 @@ public:
 		int32 SampleRate = INDEX_NONE;
 	};
 
-	DECLARE_MULTICAST_DELEGATE_OneParam(FOnOpenStream, FOnOpenStreamParams);
+	DECLARE_MULTICAST_DELEGATE_OneParam(FOnOpenStream, const FOnOpenStreamParams&);
 	virtual FDelegateHandle RegisterOpenStreamDelegate(const FOnOpenStream::FDelegate&) = 0;
 	virtual bool RemoveOpenStreamDelegate(const FDelegateHandle&) = 0;
 
@@ -80,15 +84,15 @@ public:
 
 	struct FOnRenderParams 
 	{
-		uint64 FrameID = 0;
+		uint64 BufferTickID = 0;
 		int32 NumFrames = INDEX_NONE;
 	};
 	
-	DECLARE_MULTICAST_DELEGATE_OneParam(FOnBeginRender, FOnRenderParams);
+	DECLARE_MULTICAST_DELEGATE_OneParam(FOnBeginRender, const FOnRenderParams&);
 	virtual FDelegateHandle RegisterBeginRenderDelegate(const FOnBeginRender::FDelegate&) = 0;
-	virtual bool RemoveBeginRenderDelgate(const FDelegateHandle&) = 0;
+	virtual bool RemoveBeginRenderDelegate(const FDelegateHandle&) = 0;
 
-	DECLARE_MULTICAST_DELEGATE(FOnEndRender);	
+	DECLARE_MULTICAST_DELEGATE_OneParam(FOnEndRender, const FOnRenderParams&);
 	virtual FDelegateHandle RegisterEndRenderDelegate(const FOnEndRender::FDelegate&) = 0;
-	virtual bool RemoveEndRenderDelgate(const FDelegateHandle&) = 0;
+	virtual bool RemoveEndRenderDelegate(const FDelegateHandle&) = 0;
 };
\ No newline at end of file
diff --git a/Engine/Source/Runtime/AudioLink/AudioLinkEngine/Public/IAudioLinkFactory.h b/Engine/Source/Runtime/AudioLink/AudioLinkEngine/Public/IAudioLinkFactory.h
--- a/Engine/Source/Runtime/AudioLink/AudioLinkEngine/Public/IAudioLinkFactory.h
+++ b/Engine/Source/Runtime/AudioLink/AudioLinkEngine/Public/IAudioLinkFactory.h
@@ -102,6 +102,7 @@ public:
 		int32 NumChannels = INDEX_NONE;
 		int32 SampleRate = INDEX_NONE;
 		int32 NumFramesPerBuffer = INDEX_NONE;
+		int32 TotalNumFramesInSource = INDEX_NONE;
 		UAudioLinkSettingsAbstract::FSharedSettingsProxyPtr Settings;
 	};	
 	using FAudioLinkSourcePushedSharedPtr = TSharedPtr<IAudioLinkSourcePushed, ESPMode::ThreadSafe>;
diff --git a/Engine/Source/Runtime/AudioLink/AudioMixerPlatformAudioLink/AudioMixerPlatformAudioLink.Build.cs b/Engine/Source/Runtime/AudioLink/AudioMixerPlatformAudioLink/AudioMixerPlatformAudioLink.Build.cs
--- a/Engine/Source/Runtime/AudioLink/AudioMixerPlatformAudioLink/AudioMixerPlatformAudioLink.Build.cs
+++ b/Engine/Source/Runtime/AudioLink/AudioMixerPlatformAudioLink/AudioMixerPlatformAudioLink.Build.cs
@@ -17,7 +17,8 @@ public class AudioMixerPlatformAudioLink: ModuleRules
 				"CoreUObject",
 				"Engine",
 				"AudioMixer",
-				"AudioMixerCore"
+				"AudioMixerCore",
+				"AudioLinkEngine"
 		});
 					
 		if (Target.bCompileAgainstEngine) 
diff --git a/Engine/Source/Runtime/AudioLink/AudioMixerPlatformAudioLink/Private/AudioMixerPlatformAudioLink.cpp b/Engine/Source/Runtime/AudioLink/AudioMixerPlatformAudioLink/Private/AudioMixerPlatformAudioLink.cpp
--- a/Engine/Source/Runtime/AudioLink/AudioMixerPlatformAudioLink/Private/AudioMixerPlatformAudioLink.cpp
+++ b/Engine/Source/Runtime/AudioLink/AudioMixerPlatformAudioLink/Private/AudioMixerPlatformAudioLink.cpp
@@ -10,51 +10,52 @@
 #include "Misc/CoreMisc.h"
 #include "Interfaces/ITargetPlatform.h"
 #include "Interfaces/ITargetPlatformManagerModule.h"
+#include "AudioLinkLog.h"
 
-namespace AudioMixerPlatformAudioLinkPrivate
-{	
-	using namespace Audio;
-	static const FAudioPlatformDeviceInfo& GetPlatformInfo()
+namespace Audio
+{
+	// We can't easily access AudioMixer instance from here, so find the factory in the registry.
+	static IAudioLinkFactory* GetAudioLinkFactory()
 	{
-		static const FAudioPlatformDeviceInfo DeviceInfo = []()
+		TArray<FName> Names = IAudioLinkFactory::GetAllRegisteredFactoryNames();
+		if (Names.Num())
 		{
-			// For now just hard-code a virtual audio platform, but this will ultimately match an external endpoint via audio link.
-			FAudioPlatformDeviceInfo Info;
-			Info.Name = TEXT("AudioLink Virtual Platform");
-			Info.DeviceId = Info.Name;
-			Info.SampleRate	= 48000;
-			Info.NumChannels = 8;
-			Info.bIsSystemDefault = true;
-			Info.Format = EAudioMixerStreamDataFormat::Float;
-			Info.OutputChannelArray.Add(EAudioMixerChannel::FrontLeft);
-			Info.OutputChannelArray.Add(EAudioMixerChannel::FrontRight);
-			Info.OutputChannelArray.Add(EAudioMixerChannel::FrontCenter);
-			Info.OutputChannelArray.Add(EAudioMixerChannel::LowFrequency);
-			Info.OutputChannelArray.Add(EAudioMixerChannel::SideLeft);
-			Info.OutputChannelArray.Add(EAudioMixerChannel::SideRight);
-			Info.OutputChannelArray.Add(EAudioMixerChannel::BackLeft);
-			Info.OutputChannelArray.Add(EAudioMixerChannel::BackRight);
-
-			return Info;
-		}();
-		return DeviceInfo;
+			return IAudioLinkFactory::FindFactory(Names[0]);
+		}
+		return nullptr;
 	}
-}
 
-namespace Audio
-{		
 	FAudioMixerPlatformAudioLink::FAudioMixerPlatformAudioLink()
-	{}
+		: Factory(GetAudioLinkFactory())
+	{
+		MakeDeviceInfo(
+			8,
+			48000,
+			TEXT("AudioLink AudioMixer")
+		);
+	}
 
 	bool FAudioMixerPlatformAudioLink::InitializeHardware()
 	{
 		if (IAudioMixer::ShouldRecycleThreads())
 		{
-			// Pre-create the null render device thread on XAudio2, so we can simple wake it up when we need it.
+			// Pre-create the null render device thread, so we can simple wake it up when we need it.
 			// Give it nothing to do, with a slow tick as the default, but ask it to wait for a signal to wake up.
 			CreateNullDeviceThread([] {}, 1.0f, true);
 		}
 
+		SynchronizeLink = Factory->CreateSynchronizerAudioLink();
+		if (SynchronizeLink.IsValid())
+		{
+			using FThisType = FAudioMixerPlatformAudioLink;
+			using FSync = IAudioLinkSynchronizer;
+		
+			// Register callbacks.
+			SynchronizeLink->RegisterBeginRenderDelegate(FSync::FOnBeginRender::FDelegate::CreateRaw(this, &FThisType::OnLinkRenderBegin));
+			SynchronizeLink->RegisterEndRenderDelegate(FSync::FOnEndRender::FDelegate::CreateRaw(this, &FThisType::OnLinkRenderEnd));
+			SynchronizeLink->RegisterOpenStreamDelegate(FSync::FOnOpenStream::FDelegate::CreateRaw(this, &FThisType::OnLinkOpenStream));
+		}
+
 		bInitialized = true;
 		return true;
 	}
@@ -63,6 +64,9 @@ namespace Audio
 	{		
 		StopAudioStream();
 		CloseAudioStream();
+
+		// Kill synchronizer, and unregister callbacks.
+		SynchronizeLink.Reset();
 		return true;
 	}
 
@@ -79,9 +83,7 @@ namespace Audio
 	
 	bool FAudioMixerPlatformAudioLink::GetOutputDeviceInfo(const uint32 InDeviceIndex, FAudioPlatformDeviceInfo& OutInfo)
 	{	
-		using namespace AudioMixerPlatformAudioLinkPrivate;
-		
-		OutInfo = GetPlatformInfo();
+		OutInfo = DeviceInfo;
 		return true;
 	}
 
@@ -141,7 +143,10 @@ namespace Audio
 		// Start generating audio
 		BeginGeneratingAudio();
 
-		StartRunningNullDevice();
+		if (!SynchronizeLink.IsValid())
+		{
+			StartRunningNullDevice();
+		}
 
 		AudioStreamInfo.StreamState = EAudioOutputStreamState::Running;
 
@@ -175,8 +180,7 @@ namespace Audio
 
 	FString FAudioMixerPlatformAudioLink::GetDefaultDeviceName()
 	{
-		using namespace AudioMixerPlatformAudioLinkPrivate;
-		return GetPlatformInfo().Name;
+		return DeviceInfo.Name;
 	}
 	
 	FAudioPlatformSettings FAudioMixerPlatformAudioLink::GetPlatformSettings() const
@@ -187,4 +191,51 @@ namespace Audio
 		return FAudioPlatformSettings();
 #endif // WITH_ENGINE
 	}
+
+	void FAudioMixerPlatformAudioLink::MakeDeviceInfo(int32 InNumChannels, int32 InSampleRate, const FString& InName)
+	{
+		DeviceInfo.Reset();
+		DeviceInfo.Name = InName;
+		DeviceInfo.DeviceId = InName;
+		DeviceInfo.SampleRate = InSampleRate;
+		DeviceInfo.NumChannels = InNumChannels;
+		DeviceInfo.bIsSystemDefault = true;
+		DeviceInfo.Format = EAudioMixerStreamDataFormat::Float;
+		
+		DeviceInfo.OutputChannelArray.SetNum(InNumChannels);
+		for (int32 i = 0; i < InNumChannels; ++i)
+		{
+			DeviceInfo.OutputChannelArray.Add(EAudioMixerChannel::Type(i));
+		}
+	}
+
+	void FAudioMixerPlatformAudioLink::OnLinkOpenStream(const IAudioLinkSynchronizer::FOnOpenStreamParams& InParams)
+	{
+		DeviceInfo.Name = InParams.Name;
+		DeviceInfo.NumChannels = InParams.NumChannels;
+		DeviceInfo.SampleRate = InParams.SampleRate;
+	}
+
+	void FAudioMixerPlatformAudioLink::OnLinkRenderBegin(const IAudioLinkSynchronizer::FOnRenderParams& InParams)
+	{		
+	}
+	void FAudioMixerPlatformAudioLink::OnLinkRenderEnd(const IAudioLinkSynchronizer::FOnRenderParams& InParams)
+	{
+		UE_LOG(LogAudioLink, VeryVerbose, TEXT(
+			"FAudioMixerPlatformAudioLink::OnLinkRenderEnd, TickID=%d, FramesMade=%d, LastBufferTickID=%d, FrameRemainder=%d, AudioMixer.NumFrames=%d, AudioMixer.NumBuffers=%dm This=0x%p"), 
+				InParams.BufferTickID, InParams.NumFrames, LastBufferTickID, FrameRemainder, AudioStreamInfo.NumOutputFrames, AudioStreamInfo.NumBuffers, this);
+
+		if (LastBufferTickID < InParams.BufferTickID)
+		{
+			LastBufferTickID = InParams.BufferTickID;
+
+			FrameRemainder += InParams.NumFrames;
+
+			if (FrameRemainder >= AudioStreamInfo.NumOutputFrames)
+			{
+				FrameRemainder = 0;
+				ReadNextBuffer();
+			}
+		}
+	}
 }
diff --git a/Engine/Source/Runtime/AudioLink/AudioMixerPlatformAudioLink/Private/AudioMixerPlatformAudioLink.h b/Engine/Source/Runtime/AudioLink/AudioMixerPlatformAudioLink/Private/AudioMixerPlatformAudioLink.h
--- a/Engine/Source/Runtime/AudioLink/AudioMixerPlatformAudioLink/Private/AudioMixerPlatformAudioLink.h
+++ b/Engine/Source/Runtime/AudioLink/AudioMixerPlatformAudioLink/Private/AudioMixerPlatformAudioLink.h
@@ -13,6 +13,7 @@ namespace Audio
 
 		FAudioMixerPlatformAudioLink();
 		virtual ~FAudioMixerPlatformAudioLink() = default;
+	protected:
 
 		//~ Begin IAudioMixerPlatformInterface Interface
 		FString GetPlatformApi() const override { return TEXT("AudioLink"); }
@@ -32,9 +33,22 @@ namespace Audio
 		//~ End IAudioMixerPlatformInterface Interface
 
 	private:
+		void MakeDeviceInfo(int32 InNumChannels, int32 InSampleRate, const FString& InName);
+
+		void OnLinkSuspend();
+		void OnLinkResume();
+		void OnLinkCloseStream();
+		void OnLinkOpenStream(const IAudioLinkSynchronizer::FOnOpenStreamParams&);
+		void OnLinkRenderBegin(const IAudioLinkSynchronizer::FOnRenderParams&);
+		void OnLinkRenderEnd(const IAudioLinkSynchronizer::FOnRenderParams&);
+
 		bool bSuspended = false;
 		bool bInitialized = false;
+		IAudioLinkFactory* Factory = nullptr;
 		IAudioLinkFactory::FAudioLinkSynchronizerSharedPtr SynchronizeLink;
+		FAudioPlatformDeviceInfo DeviceInfo;
+		uint32 LastBufferTickID = 0;
+		int32 FrameRemainder = 0;
 	};
 }
 
diff --git a/Engine/Source/Runtime/AudioMixer/Private/AudioMixerSource.cpp b/Engine/Source/Runtime/AudioMixer/Private/AudioMixerSource.cpp
--- a/Engine/Source/Runtime/AudioMixer/Private/AudioMixerSource.cpp
+++ b/Engine/Source/Runtime/AudioMixer/Private/AudioMixerSource.cpp
@@ -596,7 +596,7 @@ namespace Audio
 			InitParams.SourceBufferListener = WaveInstance->SourceBufferListener;
 			InitParams.bShouldSourceBufferListenerZeroBuffer = WaveInstance->bShouldSourceBufferListenerZeroBuffer;
 
-			if( WaveInstance->bShouldUseAudioLink)
+			if (WaveInstance->bShouldUseAudioLink)
 			{
 				if (IAudioLinkFactory* LinkFactory = MixerDevice->GetAudioLinkFactory())
 				{				
@@ -614,6 +614,7 @@ namespace Audio
 					CreateArgs.NumChannels = SoundBuffer->NumChannels;
 					CreateArgs.NumFramesPerBuffer = MixerDevice->GetBufferLength();
 					CreateArgs.SampleRate = MixerDevice->GetSampleRate();
+					CreateArgs.TotalNumFramesInSource = NumTotalFrames;
 					AudioLink = LinkFactory->CreateSourcePushedAudioLink(CreateArgs);
 					InitParams.AudioLink = AudioLink;
 				}
diff --git a/Engine/Source/Runtime/AudioMixer/Private/AudioMixerSourceManager.cpp b/Engine/Source/Runtime/AudioMixer/Private/AudioMixerSourceManager.cpp
--- a/Engine/Source/Runtime/AudioMixer/Private/AudioMixerSourceManager.cpp
+++ b/Engine/Source/Runtime/AudioMixer/Private/AudioMixerSourceManager.cpp
@@ -580,6 +580,12 @@ namespace Audio
 			MixerDevice->ReverbPluginInterface->OnReleaseSource(SourceId);
 		}
 
+		if (SourceInfo.AudioLink)
+		{
+			SourceInfo.AudioLink->OnSourceReleased(SourceId);
+			SourceInfo.AudioLink.Reset();
+		}
+
 		// Delete the source effects
 		SourceInfo.SourceEffectChainId = INDEX_NONE;
 		ResetSourceEffectChain(SourceId);
@@ -606,12 +612,6 @@ namespace Audio
 		SourceInfo.LowPassFreq = MAX_FILTER_FREQUENCY;
 		SourceInfo.HighPassFreq = MIN_FILTER_FREQUENCY;
 
-		if (SourceInfo.AudioLink)
-		{
-			SourceInfo.AudioLink->OnSourceReleased(SourceId);
-			SourceInfo.AudioLink.Reset();
-		}
-
 		if (SourceInfo.SourceBufferListener)
 		{
 			SourceInfo.SourceBufferListener->OnSourceReleased(SourceId);
@@ -1274,6 +1274,10 @@ namespace Audio
 			{
 				SourceInfo.SourceListener->OnDone();
 			}
+			if (SourceInfo.AudioLink)
+			{
+				SourceInfo.AudioLink->OnSourceDone(SourceId);
+			}
 		}
 	}
 
@@ -3026,6 +3030,11 @@ namespace Audio
 
 					// Notify that we're now done with this source
 					SourceInfo.SourceListener->OnDone();
+
+					if (SourceInfo.AudioLink)
+					{
+						SourceInfo.AudioLink->OnSourceDone(SourceId);
+					}
 				}
 			}
 		}
@@ -3042,6 +3051,10 @@ namespace Audio
 				SourceInfo.bIsStopping = false;
 				SourceInfo.bIsDone = true;
 				SourceInfo.SourceListener->OnDone();
+				if (SourceInfo.AudioLink)
+				{
+					SourceInfo.AudioLink->OnSourceDone(SourceId);
+				}
 			}
 
 		}
diff --git a/Engine/Source/Runtime/Engine/Classes/Sound/SoundAttenuation.h b/Engine/Source/Runtime/Engine/Classes/Sound/SoundAttenuation.h
--- a/Engine/Source/Runtime/Engine/Classes/Sound/SoundAttenuation.h
+++ b/Engine/Source/Runtime/Engine/Classes/Sound/SoundAttenuation.h
@@ -5,6 +5,7 @@
 #include "Engine/Attenuation.h"
 #include "IAudioExtensionPlugin.h"
 #include "IAudioParameterInterfaceRegistry.h"
+#include "AudioLinkSettingsAbstract.h"
 
 #include "SoundAttenuation.generated.h"
 
@@ -203,10 +204,18 @@ struct ENGINE_API FSoundAttenuationSettings : public FBaseAttenuationSettings
 	UPROPERTY(EditAnywhere, BlueprintReadWrite, Category = AttenuationSourceDataOverride, meta = (DisplayName = "Enable Source Data Override"))
 	uint8 bEnableSourceDataOverride : 1;
 
+	/** Enables/Disables AudioLink on all sources using this attenuation */
+	UPROPERTY(EditAnywhere, BlueprintReadWrite, Category = AttenuationAudioLink, meta = (DisplayName = "Enable Send to AudioLink"))
+	uint8 bEnableSendToAudioLink : 1;
+
 	/** What method we use to spatialize the sound. */
 	UPROPERTY(EditAnywhere, BlueprintReadWrite, Category = AttenuationSpatialization, meta = (ClampMin = "0", EditCondition = "bSpatialize", DisplayName = "Spatialization Method"))
 	TEnumAsByte<enum ESoundSpatializationAlgorithm> SpatializationAlgorithm;
 
+	/** AudioLink Setting Overrides */
+	UPROPERTY(EditAnywhere, BlueprintReadWrite, Category = AttenuationAudioLink, meta = (DisplayName = "AudioLink Settings Override", EditCondition = "bEnableSendToAudioLink"))
+	TObjectPtr<UAudioLinkSettingsAbstract> AudioLinkSettingsOverride;
+
 	/** What min radius to use to swap to non-binaural audio when a sound starts playing. */
 	UPROPERTY(EditAnywhere, BlueprintReadWrite, Category = AttenuationSpatialization, meta = (ClampMin = "0", EditCondition = "bSpatialize"))
 	float BinauralRadius;
@@ -413,7 +422,9 @@ struct ENGINE_API FSoundAttenuationSettings : public FBaseAttenuationSettings
 		, bEnableLogFrequencyScaling(false)
 		, bEnableSubmixSends(false)
 		, bEnableSourceDataOverride(false)
+		, bEnableSendToAudioLink(true)
 		, SpatializationAlgorithm(ESoundSpatializationAlgorithm::SPATIALIZATION_Default)
+		, AudioLinkSettingsOverride(nullptr)
 		, BinauralRadius(0.0f)
 		, AbsorptionMethod(EAirAbsorptionMethod::Linear)
 		, OcclusionTraceChannel(ECC_Visibility)
diff --git a/Engine/Source/Runtime/Engine/Private/ActiveSound.cpp b/Engine/Source/Runtime/Engine/Private/ActiveSound.cpp
--- a/Engine/Source/Runtime/Engine/Private/ActiveSound.cpp
+++ b/Engine/Source/Runtime/Engine/Private/ActiveSound.cpp
@@ -1787,6 +1787,11 @@ void FActiveSound::UpdateAttenuation(float DeltaTime, FSoundParseParameters& Par
 		}
 	}
 
+	if (Settings->AudioLinkSettingsOverride)
+	{
+		ParseParams.AudioLinkSettingsOverride = Settings->AudioLinkSettingsOverride;
+	}
+
 	// Attenuate with the absorption filter if necessary
 	if (Settings->bAttenuateWithLPF)
 	{
@@ -1814,6 +1819,7 @@ void FActiveSound::UpdateAttenuation(float DeltaTime, FSoundParseParameters& Par
 	ParseParams.bApplyNormalizationToStereoSounds = Settings->bApplyNormalizationToStereoSounds;
 	ParseParams.bUseSpatialization |= Settings->bSpatialize;
 	ParseParams.bEnableSourceDataOverride |= Settings->bEnableSourceDataOverride;
+	ParseParams.bEnableSendToAudioLink |= Settings->bEnableSendToAudioLink;
 
 	// Check the binaural radius to determine if we're going to HRTF spatialize, cache the result
 	if (bIsFirstAttenuationUpdate)
diff --git a/Engine/Source/Runtime/Engine/Private/SoundAttenuation.cpp b/Engine/Source/Runtime/Engine/Private/SoundAttenuation.cpp
--- a/Engine/Source/Runtime/Engine/Private/SoundAttenuation.cpp
+++ b/Engine/Source/Runtime/Engine/Private/SoundAttenuation.cpp
@@ -112,6 +112,7 @@ bool FSoundAttenuationSettings::operator==(const FSoundAttenuationSettings& Othe
 			&& bEnableLogFrequencyScaling == Other.bEnableLogFrequencyScaling
 			&& bEnableSubmixSends 		== Other.bEnableSubmixSends
 			&& bEnableListenerFocus 	== Other.bEnableListenerFocus
+			&& bEnableSendToAudioLink	== Other.bEnableSendToAudioLink
 			&& FocusAzimuth				== Other.FocusAzimuth
 			&& NonFocusAzimuth			== Other.NonFocusAzimuth
 			&& FocusDistanceScale		== Other.FocusDistanceScale
@@ -127,6 +128,7 @@ bool FSoundAttenuationSettings::operator==(const FSoundAttenuationSettings& Othe
 			&& bEnableReverbSend		== Other.bEnableReverbSend
 			&& PluginSettings.ReverbPluginSettingsArray		== Other.PluginSettings.ReverbPluginSettingsArray
 			&& PluginSettings.SourceDataOverridePluginSettingsArray == Other.PluginSettings.SourceDataOverridePluginSettingsArray
+			&& AudioLinkSettingsOverride == Other.AudioLinkSettingsOverride
 			&& ReverbWetLevelMin		== Other.ReverbWetLevelMin
 			&& ReverbWetLevelMax		== Other.ReverbWetLevelMax
 			&& ReverbDistanceMin		== Other.ReverbDistanceMin
diff --git a/Engine/Source/Runtime/Engine/Private/SoundWave.cpp b/Engine/Source/Runtime/Engine/Private/SoundWave.cpp
--- a/Engine/Source/Runtime/Engine/Private/SoundWave.cpp
+++ b/Engine/Source/Runtime/Engine/Private/SoundWave.cpp
@@ -3059,6 +3059,9 @@ void USoundWave::Parse(FAudioDevice* AudioDevice, const UPTRINT NodeWaveInstance
 
 	WaveInstance->bEnableSourceDataOverride = ParseParams.bEnableSourceDataOverride;
 
+	WaveInstance->bShouldUseAudioLink = ParseParams.bEnableSendToAudioLink;
+	WaveInstance->AudioLinkSettingsOverride = ParseParams.AudioLinkSettingsOverride;
+
 	if (WaveInstance->IsPlaying())
 	{
 		WaveInstances.Add(WaveInstance);
diff --git a/Engine/Source/Runtime/Engine/Public/ActiveSound.h b/Engine/Source/Runtime/Engine/Public/ActiveSound.h
--- a/Engine/Source/Runtime/Engine/Public/ActiveSound.h
+++ b/Engine/Source/Runtime/Engine/Public/ActiveSound.h
@@ -160,6 +160,9 @@ struct FSoundParseParameters
 	// What source data override plugin source settings to use
 	USourceDataOverridePluginSourceSettingsBase* SourceDataOverridePluginSettings;
 
+	// If using AudioLink, this allows the settings to be overriden.
+	UAudioLinkSettingsAbstract* AudioLinkSettingsOverride = nullptr;
+
 	// What source effect chain to use
 	USoundEffectSourcePresetChain* SourceEffectChain;
 
@@ -189,6 +192,8 @@ struct FSoundParseParameters
 
 	uint32 bEnableSourceDataOverride : 1;
 
+	uint32 bEnableSendToAudioLink : 1;
+
 	// Whether the sound should be spatialized
 	uint8 bUseSpatialization:1;
 
@@ -239,6 +244,7 @@ struct FSoundParseParameters
 		, OcclusionPluginSettings(nullptr)
 		, ReverbPluginSettings(nullptr)
 		, SourceDataOverridePluginSettings(nullptr)
+		, AudioLinkSettingsOverride(nullptr)
 		, SourceEffectChain(nullptr)
 		, LowPassFilterFrequency(MAX_FILTER_FREQUENCY)
 		, AttenuationLowpassFilterFrequency(MAX_FILTER_FREQUENCY)
