AudioLink PlatformConfig overrides to match host

diff --git a/Engine/Source/Runtime/AudioLink/AudioLinkEngine/Public/IAudioLink.h b/Engine/Source/Runtime/AudioLink/AudioLinkEngine/Public/IAudioLink.h
--- a/Engine/Source/Runtime/AudioLink/AudioLinkEngine/Public/IAudioLink.h
+++ b/Engine/Source/Runtime/AudioLink/AudioLinkEngine/Public/IAudioLink.h
@@ -68,10 +68,30 @@ public:
 
 	struct FOnOpenStreamParams
 	{
+		/**!
+		 * Name of Implementation.
+		 */
 		FString Name;
-		int32 NumFrames = INDEX_NONE; 
+
+		/**
+		 * Number of frames per callback buffer. INDEX_NONE if not set.
+		 */
+		int32 NumFrames = INDEX_NONE;
+
+		/**
+		 * Number of channels we're rendering at i.e. Stereo=2. INDEX_NONE if not set.
+		 */
 		int32 NumChannels = INDEX_NONE;
+
+		/**
+		 * Sample rate of AL implementation. INDEX_NONE if not set.
+		 */
 		int32 SampleRate = INDEX_NONE;
+
+		/**
+		 * Num of Sources that AL is limited to. INDEX_NONE if not set.
+		 */
+		int32 NumSources = INDEX_NONE;
 	};
 
 	DECLARE_MULTICAST_DELEGATE_OneParam(FOnOpenStream, const FOnOpenStreamParams&);
@@ -95,4 +115,10 @@ public:
 	DECLARE_MULTICAST_DELEGATE_OneParam(FOnEndRender, const FOnRenderParams&);
 	virtual FDelegateHandle RegisterEndRenderDelegate(const FOnEndRender::FDelegate&) = 0;
 	virtual bool RemoveEndRenderDelegate(const FDelegateHandle&) = 0;
+
+	/**
+	 * Gets any OpenStreamParams that were sent ahead of synchronizing existing.
+	 * @return TOptional<FOpenStreamParams> (set if data is available).
+	 */
+	virtual TOptional<FOnOpenStreamParams> GetCachedOpenStreamParams() const { return {}; }
 };
diff --git a/Engine/Source/Runtime/AudioLink/AudioMixerPlatformAudioLink/Private/AudioMixerPlatformAudioLink.cpp b/Engine/Source/Runtime/AudioLink/AudioMixerPlatformAudioLink/Private/AudioMixerPlatformAudioLink.cpp
--- a/Engine/Source/Runtime/AudioLink/AudioMixerPlatformAudioLink/Private/AudioMixerPlatformAudioLink.cpp
+++ b/Engine/Source/Runtime/AudioLink/AudioMixerPlatformAudioLink/Private/AudioMixerPlatformAudioLink.cpp
@@ -44,16 +44,20 @@ namespace Audio
 			CreateNullDeviceThread([] {}, 1.0f, true);
 		}
 
-		SynchronizeLink = Factory->CreateSynchronizerAudioLink();
-		if (SynchronizeLink.IsValid())
+		if (IAudioLinkSynchronizer* Synchronizer = GetOrCreateSynchronizer())
 		{
 			using FThisType = FAudioMixerPlatformAudioLink;
 			using FSync = IAudioLinkSynchronizer;
 		
 			// Register callbacks.
-			SynchronizeLink->RegisterBeginRenderDelegate(FSync::FOnBeginRender::FDelegate::CreateRaw(this, &FThisType::OnLinkRenderBegin));
-			SynchronizeLink->RegisterEndRenderDelegate(FSync::FOnEndRender::FDelegate::CreateRaw(this, &FThisType::OnLinkRenderEnd));
-			SynchronizeLink->RegisterOpenStreamDelegate(FSync::FOnOpenStream::FDelegate::CreateRaw(this, &FThisType::OnLinkOpenStream));
+			RenderBeginHandle = Synchronizer->RegisterBeginRenderDelegate(FSync::FOnBeginRender::FDelegate::CreateRaw(this, &FThisType::OnLinkRenderBegin));
+			RenderEndHandle = Synchronizer->RegisterEndRenderDelegate(FSync::FOnEndRender::FDelegate::CreateRaw(this, &FThisType::OnLinkRenderEnd));
+			OpenStreamHandle = Synchronizer->RegisterOpenStreamDelegate(FSync::FOnOpenStream::FDelegate::CreateRaw(this, &FThisType::OnLinkOpenStream));
+
+			if (const TOptional Params = Synchronizer->GetCachedOpenStreamParams(); Params.IsSet())
+			{
+				OnLinkOpenStream(*Params);
+			}
 		}
 
 		bInitialized = true;
@@ -66,6 +70,15 @@ namespace Audio
 		CloseAudioStream();
 
 		// Kill synchronizer, and unregister callbacks.
+		if (SynchronizeLink.IsValid())
+		{
+			SynchronizeLink->RemoveBeginRenderDelegate(RenderBeginHandle);
+			SynchronizeLink->RemoveEndRenderDelegate(RenderEndHandle);
+			SynchronizeLink->RemoveOpenStreamDelegate(OpenStreamHandle);
+			RenderBeginHandle.Reset();
+			RenderEndHandle.Reset();
+			OpenStreamHandle.Reset();
+		}
 		SynchronizeLink.Reset();
 		return true;
 	}
@@ -150,6 +163,8 @@ namespace Audio
 
 		AudioStreamInfo.StreamState = EAudioOutputStreamState::Running;
 
+		bAtomicStreamRunning = true; 
+
 		return true;
 	}
 
@@ -169,6 +184,8 @@ namespace Audio
 
 			check(AudioStreamInfo.StreamState == EAudioOutputStreamState::Stopped);
 		}
+		
+		bAtomicStreamRunning = false;
 
 		return true;
 	}
@@ -182,13 +199,55 @@ namespace Audio
 	{
 		return DeviceInfo.Name;
 	}
+
+	IAudioLinkSynchronizer* FAudioMixerPlatformAudioLink::GetOrCreateSynchronizer() const
+	{
+		if (!SynchronizeLink.IsValid())
+		{
+			check(Factory);
+			SynchronizeLink = Factory->CreateSynchronizerAudioLink();
+		}
+		return SynchronizeLink.Get();
+	}
 	
 	FAudioPlatformSettings FAudioMixerPlatformAudioLink::GetPlatformSettings() const
 	{
 #if WITH_ENGINE
-		return FAudioPlatformSettings::GetPlatformSettings(FPlatformProperties::GetRuntimeSettingsClassName());
+		FAudioPlatformSettings Settings = FAudioPlatformSettings::GetPlatformSettings(FPlatformProperties::GetRuntimeSettingsClassName());
+		
+		if (const IAudioLinkSynchronizer* Synchronizer = GetOrCreateSynchronizer())
+		{
+			if (TOptional<IAudioLinkSynchronizer::FOnOpenStreamParams> Params = Synchronizer->GetCachedOpenStreamParams(); Params.IsSet())
+			{
+				// Override the .INI settings with the host audio settings, assuming we're set up.
+				// Check if each setting has been passed by AudioLink implementation. It will be INDEX_NONE otherwise.
+
+				// NOTE. 'MaxChannels' on the settings is really MaxSources.
+				if (Params->NumSources != INDEX_NONE)
+				{
+					Settings.MaxChannels = Params->NumSources; 
+				}
+				if (Params->SampleRate != INDEX_NONE)
+				{
+					Settings.SampleRate = Params->SampleRate; 
+				}
+				if (Params->NumFrames != INDEX_NONE) 
+				{
+					Settings.CallbackBufferFrameSize = Params->NumFrames;
+				}
+			}
+			else
+			{
+				UE_LOG(LogAudioLink, Warning, TEXT("%hs - OpenStreamParams has not been set, all settings will be from the Default Engine settings, not the host."), __func__);
+			}
+		}
+		else
+		{
+			UE_LOG(LogAudioLink, Warning, TEXT("%hs - Synchronizer does not exist yet."), __func__);
+		}
+		return Settings;;
 #else
-		return FAudioPlatformSettings();
+		return {};
 #endif // WITH_ENGINE
 	}
 
@@ -211,9 +270,11 @@ namespace Audio
 
 	void FAudioMixerPlatformAudioLink::OnLinkOpenStream(const IAudioLinkSynchronizer::FOnOpenStreamParams& InParams)
 	{
-		DeviceInfo.Name = InParams.Name;
-		DeviceInfo.NumChannels = InParams.NumChannels;
-		DeviceInfo.SampleRate = InParams.SampleRate;
+		MakeDeviceInfo(
+			InParams.NumChannels,
+			InParams.SampleRate,
+			TEXT("AudioLink AudioMixer"));
+
 	}
 
 	void FAudioMixerPlatformAudioLink::OnLinkRenderBegin(const IAudioLinkSynchronizer::FOnRenderParams& InParams)
@@ -222,19 +283,22 @@ namespace Audio
 	void FAudioMixerPlatformAudioLink::OnLinkRenderEnd(const IAudioLinkSynchronizer::FOnRenderParams& InParams)
 	{
 		UE_LOG(LogAudioLink, VeryVerbose, TEXT(
-			"FAudioMixerPlatformAudioLink::OnLinkRenderEnd, TickID=%d, FramesMade=%d, LastBufferTickID=%d, FrameRemainder=%d, AudioMixer.NumFrames=%d, AudioMixer.NumBuffers=%dm This=0x%p"), 
+			"FAudioMixerPlatformAudioLink::OnLinkRenderEnd, TickID=%llu, FramesMade=%d, LastBufferTickID=%d, FrameRemainder=%d, AudioMixer.NumFrames=%d, AudioMixer.NumBuffers=%dm This=0x%p"), 
 				InParams.BufferTickID, InParams.NumFrames, LastBufferTickID, FrameRemainder, AudioStreamInfo.NumOutputFrames, AudioStreamInfo.NumBuffers, this);
 
-		if (LastBufferTickID < InParams.BufferTickID)
+		if (bAtomicStreamRunning)  // Make sure Unreal is open and ready to receive input.
 		{
-			LastBufferTickID = InParams.BufferTickID;
+			if (LastBufferTickID < InParams.BufferTickID)
+			{
+				LastBufferTickID = InParams.BufferTickID;
 
-			FrameRemainder += InParams.NumFrames;
+				FrameRemainder += InParams.NumFrames;
 
-			if (FrameRemainder >= AudioStreamInfo.NumOutputFrames)
-			{
-				FrameRemainder = 0;
-				ReadNextBuffer();
+				if (FrameRemainder >= AudioStreamInfo.NumOutputFrames)
+				{
+					FrameRemainder = 0;
+					ReadNextBuffer();
+				}
 			}
 		}
 	}
diff --git a/Engine/Source/Runtime/AudioLink/AudioMixerPlatformAudioLink/Private/AudioMixerPlatformAudioLink.h b/Engine/Source/Runtime/AudioLink/AudioMixerPlatformAudioLink/Private/AudioMixerPlatformAudioLink.h
--- a/Engine/Source/Runtime/AudioLink/AudioMixerPlatformAudioLink/Private/AudioMixerPlatformAudioLink.h
+++ b/Engine/Source/Runtime/AudioLink/AudioMixerPlatformAudioLink/Private/AudioMixerPlatformAudioLink.h
@@ -35,20 +35,28 @@ namespace Audio
 	private:
 		void MakeDeviceInfo(int32 InNumChannels, int32 InSampleRate, const FString& InName);
 
-		void OnLinkSuspend();
-		void OnLinkResume();
-		void OnLinkCloseStream();
 		void OnLinkOpenStream(const IAudioLinkSynchronizer::FOnOpenStreamParams&);
 		void OnLinkRenderBegin(const IAudioLinkSynchronizer::FOnRenderParams&);
 		void OnLinkRenderEnd(const IAudioLinkSynchronizer::FOnRenderParams&);
 
+		IAudioLinkSynchronizer* GetOrCreateSynchronizer() const;
+
 		bool bSuspended = false;
 		bool bInitialized = false;
+		
+		// Atomic flag that's set by unreal, but read in callback from other host engines.
+		std::atomic<bool> bAtomicStreamRunning;
+		
 		IAudioLinkFactory* Factory = nullptr;
-		IAudioLinkFactory::FAudioLinkSynchronizerSharedPtr SynchronizeLink;
+		mutable IAudioLinkFactory::FAudioLinkSynchronizerSharedPtr SynchronizeLink;
 		FAudioPlatformDeviceInfo DeviceInfo;
 		uint32 LastBufferTickID = 0;
 		int32 FrameRemainder = 0;
+
+		FDelegateHandle OpenStreamHandle;
+		FDelegateHandle RenderBeginHandle;
+		FDelegateHandle RenderEndHandle;
 	};
+
 }
 
