Better lifecycle response of dynamic submixes in BufferedSubmixListener.

#Synced-CL 48417707

diff --git a/Engine/Source/Runtime/AudioLink/AudioLinkEngine/Private/BufferedSubmixListener.cpp b/Engine/Source/Runtime/AudioLink/AudioLinkEngine/Private/BufferedSubmixListener.cpp
--- a/Engine/Source/Runtime/AudioLink/AudioLinkEngine/Private/BufferedSubmixListener.cpp
+++ b/Engine/Source/Runtime/AudioLink/AudioLinkEngine/Private/BufferedSubmixListener.cpp
@@ -38,15 +38,17 @@ bool FBufferedSubmixListener::Start(FAudioDevice* InAudioDevice, USoundSubmix* I
 	{
 		if (TrySetStartedFlag())
 		{
-			if (ensure(IsValid(InSubmix)))
+			USoundSubmix* QualifiedSubmix = InSubmix ? TObjectPtr<USoundSubmix>(InSubmix) : TObjectPtr<USoundSubmix>(&InAudioDevice->GetMainSubmixObject());
+			
+			check(IsInAudioThread() || IsInGameThread());
+			if (ensure(IsValid(QualifiedSubmix)))
 			{
-				InAudioDevice->RegisterSubmixBufferListener(AsShared(), *InSubmix);
-		
+				InAudioDevice->RegisterSubmixBufferListener(AsShared(), *QualifiedSubmix);
 				DeviceId = InAudioDevice->DeviceID;
-				Submix = InSubmix ? TObjectPtr<USoundSubmix>(InSubmix) : TObjectPtr<USoundSubmix>(&InAudioDevice->GetMainSubmixObject());
-				
+				Submix = QualifiedSubmix;
 				return true;
 			}
+
 		}
 	}
 	return false;
@@ -65,9 +67,12 @@ void FBufferedSubmixListener::Stop(FAudioDevice* InAudioDevice)
 		{
 			if (TryUnsetStartedFlag())
 			{
-				if (ensure(IsValid(Submix)))
+				check(IsInAudioThread() || IsInGameThread());
+
+				// In audio thread, GC cannot interrupt this task
+				if (USoundSubmix* QualifiedSubmix = Submix.Get())
 				{
-					InAudioDevice->UnregisterSubmixBufferListener(AsShared(), *Submix);
+					InAudioDevice->UnregisterSubmixBufferListener(AsShared(), *QualifiedSubmix);
 				}
 
 				DeviceId = BufferedSubmixListenerPrivate::InvalidAudioDeviceId;
diff --git a/Engine/Source/Runtime/AudioLink/AudioLinkEngine/Public/BufferedSubmixListener.h b/Engine/Source/Runtime/AudioLink/AudioLinkEngine/Public/BufferedSubmixListener.h
--- a/Engine/Source/Runtime/AudioLink/AudioLinkEngine/Public/BufferedSubmixListener.h
+++ b/Engine/Source/Runtime/AudioLink/AudioLinkEngine/Public/BufferedSubmixListener.h
@@ -44,5 +44,5 @@ private:
 	bool bZeroInputBuffer = false;
 	FString Name;
 
-	TObjectPtr<USoundSubmix> Submix;
+	TWeakObjectPtr<USoundSubmix> Submix;
 };
