Safety check before dereferencing an AudioDevice pointer passed in.

#Synced-CL 49553638

diff --git a/Engine/Source/Runtime/AudioLink/AudioLinkEngine/Private/BufferedSubmixListener.cpp b/Engine/Source/Runtime/AudioLink/AudioLinkEngine/Private/BufferedSubmixListener.cpp
--- a/Engine/Source/Runtime/AudioLink/AudioLinkEngine/Private/BufferedSubmixListener.cpp
+++ b/Engine/Source/Runtime/AudioLink/AudioLinkEngine/Private/BufferedSubmixListener.cpp
@@ -56,7 +56,11 @@ bool FBufferedSubmixListener::Start(FAudioDevice* InAudioDevice, USoundSubmix* I
 
 bool FBufferedSubmixListener::Start(FAudioDevice* InAudioDevice)
 {
-	return Start(InAudioDevice, &InAudioDevice->GetMainSubmixObject());
+	if (ensure(InAudioDevice))
+	{
+		return Start(InAudioDevice, &InAudioDevice->GetMainSubmixObject());
+	}
+	return false;
 }
 
 void FBufferedSubmixListener::Stop(FAudioDevice* InAudioDevice)
